Making override match interface by adding forward declare added in 2488444

diff --git a/Engine/Plugins/Runtime/OculusAudio/Source/OculusAudio/Private/OculusAudio.cpp b/Engine/Plugins/Runtime/OculusAudio/Source/OculusAudio/Private/OculusAudio.cpp
--- a/Engine/Plugins/Runtime/OculusAudio/Source/OculusAudio/Private/OculusAudio.cpp
+++ b/Engine/Plugins/Runtime/OculusAudio/Source/OculusAudio/Private/OculusAudio.cpp
@@ -273,7 +273,7 @@ public:
 		}
 	}
 
-	IAudioSpatializationAlgorithm* GetNewSpatializationAlgorithm(FAudioDevice* AudioDevice) override
+	IAudioSpatializationAlgorithm* GetNewSpatializationAlgorithm(class FAudioDevice* AudioDevice) override
 	{
 		return new FHrtfSpatializationAlgorithm(AudioDevice);
 	}
