Include the Oculus audio DLL in Rocket, and stage it during packaging.

diff --git a/Engine/Build/InstalledEngineFilters.ini b/Engine/Build/InstalledEngineFilters.ini
--- a/Engine/Build/InstalledEngineFilters.ini
+++ b/Engine/Build/InstalledEngineFilters.ini
@@ -50,6 +50,7 @@
 {Win64} /Engine/Binaries/ThirdParty/ICU/icu4c-53_1/Win32/VS2013/*
 {Win64} /Engine/Binaries/ThirdParty/ICU/icu4c-53_1/Win64/VS2013/*
 {Win64} /Engine/Binaries/ThirdParty/nvTextureTools/Win64/*.dll
+{Win64} /Engine/Binaries/ThirdParty/Oculus/Audio/Win64/ovraudio64.dll
 {Win64} /Engine/Binaries/ThirdParty/Ogg/Win64/VS2013/*.dll
 {Win64} /Engine/Binaries/ThirdParty/Vorbis/Win64/VS2013/*.dll
 {Win64} /Engine/Binaries/ThirdParty/PhysX/APEX-1.3/Win64/VS2013/*PROFILE_x64.dll
diff --git a/Engine/Plugins/Runtime/OculusAudio/Source/OculusAudio/OculusAudio.Build.cs b/Engine/Plugins/Runtime/OculusAudio/Source/OculusAudio/OculusAudio.Build.cs
--- a/Engine/Plugins/Runtime/OculusAudio/Source/OculusAudio/OculusAudio.Build.cs
+++ b/Engine/Plugins/Runtime/OculusAudio/Source/OculusAudio/OculusAudio.Build.cs
@@ -34,6 +34,7 @@ namespace UnrealBuildTool.Rules
             if (Target.Platform == UnrealTargetPlatform.Win64)
             {
 				PrivateDependencyModuleNames.AddRange(new string[] { "LibOVRAudio" });
+				RuntimeDependencies.Add(new RuntimeDependency("$(EngineDir)/Binaries/ThirdParty/Oculus/Audio/Win64/ovraudio64.dll"));
             }
 
             AddThirdPartyPrivateStaticDependencies(Target, "DX11Audio");
