Changing the data structure used to store oculus audio spatialization effect instances

diff --git a/Engine/Plugins/Runtime/OculusAudio/Source/OculusAudio/Private/OculusAudio.cpp b/Engine/Plugins/Runtime/OculusAudio/Source/OculusAudio/Private/OculusAudio.cpp
--- a/Engine/Plugins/Runtime/OculusAudio/Source/OculusAudio/Private/OculusAudio.cpp
+++ b/Engine/Plugins/Runtime/OculusAudio/Source/OculusAudio/Private/OculusAudio.cpp
@@ -28,10 +28,9 @@ public:
 
 	~FHrtfSpatializationAlgorithm()
 	{
-		// Release all the effects for the oculus spatialization effect
-		for (TMap<uint32, class FXAudio2HRTFEffect* >::TIterator It(HRTFEffects); It; ++It)
+		for (int32 i = 0; i < HRTFEffects.Num(); ++i)
 		{
-			It.Value()->Release();
+			HRTFEffects[i].Reset();
 		}
 
 		HRTFEffects.Empty();
@@ -121,35 +120,40 @@ public:
 	virtual bool CreateSpatializationEffect(uint32 VoiceId) override
 	{
 		// If an effect for this voice has already been created, then leave
-		if (HRTFEffects.Contains(VoiceId))
+		if (VoiceId < (uint32)HRTFEffects.Num())
 		{
 			return false;
 		}
 
-		FXAudio2HRTFEffect* NewHRTFEffect = new FXAudio2HRTFEffect(VoiceId, AudioDevice);
+		TSharedPtr<FXAudio2HRTFEffect> NewHRTFEffect = TSharedPtr<FXAudio2HRTFEffect>(new FXAudio2HRTFEffect(VoiceId, AudioDevice));
 		NewHRTFEffect->Initialize(nullptr, 0);
-		HRTFEffects.Add(VoiceId, NewHRTFEffect);
+		HRTFEffects.Add(NewHRTFEffect);
 		return true;
 	}
 
 	void* GetSpatializationEffect(uint32 VoiceId) override
 	{
-		FXAudio2HRTFEffect** OculusEffect = HRTFEffects.Find(VoiceId);
-		if (OculusEffect)
+		if (VoiceId < (uint32)HRTFEffects.Num())
 		{
-			// Trust us...
-			return (void*)(*OculusEffect);
+			TSharedPtr<FXAudio2HRTFEffect> Effect = HRTFEffects[VoiceId];
+			if (Effect.IsValid())
+			{
+				return (void*)Effect.Get();
+			}
 		}
 		return nullptr;
 	}
 
 	void SetSpatializationParameters(uint32 VoiceId, const FVector& EmitterPosition, ESpatializationEffectType AlgorithmType) override
 	{
-		FXAudio2HRTFEffect** OculusEffect = HRTFEffects.Find(VoiceId);
-		if (OculusEffect)
+		if (VoiceId < (uint32)HRTFEffects.Num())
 		{
-			FAudioHRTFEffectParameters Params(EmitterPosition);
-			(*OculusEffect)->SetParameters((const void *)&Params, sizeof(Params));
+			TSharedPtr<FXAudio2HRTFEffect> Effect = HRTFEffects[VoiceId];
+			if (Effect.IsValid())
+			{
+				FAudioHRTFEffectParameters Params(EmitterPosition);
+				Effect->SetParameters((const void *)&Params, sizeof(Params));
+			}
 		}
 	}
 
@@ -205,8 +209,7 @@ private:
 #endif
 
 	/** Xaudio2 effects for the oculus plugin */
-	TMap<uint32, class FXAudio2HRTFEffect* > HRTFEffects;
-
+	TArray<TSharedPtr<FXAudio2HRTFEffect>> HRTFEffects;
 
 };
 
