Updating Oculus Audio plugin

- Changing the way the Oculus Audio plugin deals with parameter passing, rather than using CXAPOParameterBase, we do our own parameter passing now.
- Removing COM ref counting in case of double-stopping and premature destruction

diff --git a/Engine/Plugins/Runtime/OculusAudio/Source/OculusAudio/Private/OculusAudio.cpp b/Engine/Plugins/Runtime/OculusAudio/Source/OculusAudio/Private/OculusAudio.cpp
--- a/Engine/Plugins/Runtime/OculusAudio/Source/OculusAudio/Private/OculusAudio.cpp
+++ b/Engine/Plugins/Runtime/OculusAudio/Source/OculusAudio/Private/OculusAudio.cpp
@@ -56,12 +56,15 @@ public:
 		InitializeSpatializationEffect(BufferLength);
 
 		check(HRTFEffects.Num() == 0);
+		check(Params.Num() == 0);
+
 		for (int32 EffectIndex = 0; EffectIndex < (int32)AudioDevice->MaxChannels; ++EffectIndex)
 		{
 			FXAudio2HRTFEffect* NewHRTFEffect = new FXAudio2HRTFEffect(EffectIndex, AudioDevice);
 			NewHRTFEffect->Initialize(nullptr, 0);
 			HRTFEffects.Add(NewHRTFEffect);
-			// IsStopped.Add(true);
+
+			Params.Add(FAudioSpatializationParams());
 		}
 	}
 
@@ -72,7 +75,7 @@ public:
 		{
 			if (Effect)
 			{
-				Effect->Release();
+				delete Effect;
 			}
 		}
 		HRTFEffects.Empty();
@@ -171,13 +174,20 @@ public:
 		return nullptr;
 	}
 
-	void SetSpatializationParameters(uint32 VoiceId, const FVector& EmitterPosition, ESpatializationEffectType AlgorithmType) override
+	void SetSpatializationParameters(uint32 VoiceId, const FAudioSpatializationParams& InParams) override
 	{
-		if ((int32)VoiceId < HRTFEffects.Num())
-		{
-			FAudioHRTFEffectParameters Params(EmitterPosition);
-			HRTFEffects[VoiceId]->SetParameters((const void *)&Params, sizeof(Params));
-		}
+		check((int32)VoiceId < Params.Num());
+
+		FScopeLock ScopeLock(&ParamCriticalSection);
+		Params[VoiceId] = InParams;
+	}
+
+	void GetSpatializationParameters(uint32 VoiceId, FAudioSpatializationParams& OutParams) override
+	{
+		check((int32)VoiceId < Params.Num());
+
+		FScopeLock ScopeLock(&ParamCriticalSection);
+		OutParams = Params[VoiceId];
 	}
 
 private:
@@ -222,6 +232,9 @@ private:
 
 	/** Xaudio2 effects for the oculus plugin */
 	TArray<class FXAudio2HRTFEffect*> HRTFEffects;
+	TArray<FAudioSpatializationParams> Params;
+
+	FCriticalSection ParamCriticalSection;
 };
 
 
diff --git a/Engine/Plugins/Runtime/OculusAudio/Source/OculusAudio/Private/OculusAudioEffect.cpp b/Engine/Plugins/Runtime/OculusAudio/Source/OculusAudio/Private/OculusAudioEffect.cpp
--- a/Engine/Plugins/Runtime/OculusAudio/Source/OculusAudio/Private/OculusAudioEffect.cpp
+++ b/Engine/Plugins/Runtime/OculusAudio/Source/OculusAudio/Private/OculusAudioEffect.cpp
@@ -16,19 +16,22 @@
 #include "HideWindowsPlatformTypes.h"
 
 FXAudio2HRTFEffect::FXAudio2HRTFEffect(uint32 InVoiceId, FAudioDevice* InAudioDevice)
-	: CXAPOParametersBase(&Registration, (uint8*)HRTFEffectParameters, sizeof(FAudioHRTFEffectParameters), false)
+	: CXAPOBase(&Registration)
 	, VoiceId(InVoiceId)
 	, AudioDevice(InAudioDevice)
 	, bPassThrough(false)
 {
-	FAudioHRTFEffectParameters DefaultParameters;
-	SetParameters(&DefaultParameters, sizeof(DefaultParameters));
+}
+
+FXAudio2HRTFEffect::~FXAudio2HRTFEffect()
+{
+	// nothing
 }
 
 HRESULT FXAudio2HRTFEffect::LockForProcess(UINT32 InputLockedParameterCount, const XAPO_LOCKFORPROCESS_BUFFER_PARAMETERS *pInputLockedParameters, UINT32 OutputLockedParameterCount, const XAPO_LOCKFORPROCESS_BUFFER_PARAMETERS *pOutputLockedParameters)
 {
 	// Try to lock the process on the base class before attempting to do any initialization here. 
-	HRESULT Result = CXAPOParametersBase::LockForProcess(InputLockedParameterCount, pInputLockedParameters, OutputLockedParameterCount, pOutputLockedParameters);
+	HRESULT Result = CXAPOBase::LockForProcess(InputLockedParameterCount, pInputLockedParameters, OutputLockedParameterCount, pOutputLockedParameters);
 
 	// Process couldn't be locked. ABORT! 
 	if (FAILED(Result))
@@ -62,7 +65,7 @@ void FXAudio2HRTFEffect::Process(UINT32 InputProcessParameterCount, const XAPO_P
 	check(InputProcessParameterCount == 1);
 	check(OutputProcessParameterCount == 1);
 
-	// Make sure our input and output buffers are *not* the same. This is not an in-process plugin
+	// Make sure our input and output buffers are *not* the same. This is not an in-place plugin
 	check(pInputProcessParameters[0].pBuffer != pOutputProcessParameters[0].pBuffer);
 
 	// Don't do anything if we aren't enabled or there is no output volume
@@ -71,12 +74,11 @@ void FXAudio2HRTFEffect::Process(UINT32 InputProcessParameterCount, const XAPO_P
 		return;
 	}
 
-	// This locks the parameters until the EndProcess() call at the bottom. 
-	FAudioHRTFEffectParameters* CurrentParameters = (FAudioHRTFEffectParameters*)BeginProcess();
+	FAudioSpatializationParams CurrentParameters;
+	AudioDevice->SpatializeProcessor->GetSpatializationParameters(VoiceId, CurrentParameters);
 
-	// We set the emitter position and the current algorithm from outside the XAPO plugin
-	FVector & EmitterPosition = CurrentParameters->EmitterPosition;
-	ESpatializationEffectType SpatializationAlgorithm = CurrentParameters->SpatializationAlgorithm;
+	FVector& EmitterPosition = CurrentParameters.EmitterPosition;
+	ESpatializationEffectType SpatializationAlgorithm = CurrentParameters.SpatializationAlgorithm;
 
 	// We better not have default set at this point because default means "no HRTF processing". We also
 	// better thave the plugin because this XAPO code is inside the plugin
@@ -88,59 +90,58 @@ void FXAudio2HRTFEffect::Process(UINT32 InputProcessParameterCount, const XAPO_P
 
 	switch (InFlags)
 	{
-	case XAPO_BUFFER_VALID:
-	{
-		// Get the input/output samples
-		float* InputSamples = (float*)pInputProcessParameters[0].pBuffer;
-		float* OutputSamples = (float *)pOutputProcessParameters[0].pBuffer;
-
-		// If we're in pass through mode, we're not going to the oculus SDK, but instead, we're going to
-		// just split the mono stream into a stereo stream.
-		if (bPassThrough)
+		case XAPO_BUFFER_VALID:
 		{
-			uint32 SampleIndex = 0;
-			for (uint32 FrameIndex = 0; FrameIndex < InputFrameCount; ++FrameIndex)
+			// Get the input/output samples
+			float* InputSamples = (float*)pInputProcessParameters[0].pBuffer;
+			float* OutputSamples = (float *)pOutputProcessParameters[0].pBuffer;
+
+			// If we're in pass through mode, we're not going to the oculus SDK, but instead, we're going to
+			// just split the mono stream into a stereo stream.
+			if (bPassThrough)
 			{
-				// Scale by 0.5f to preserve power when splitting into stereo
-				float InputSample = InputSamples[FrameIndex];
-				OutputSamples[SampleIndex++] = 0.5f * InputSample;
-				OutputSamples[SampleIndex++] = 0.5f * InputSample;
+				uint32 SampleIndex = 0;
+				for (uint32 FrameIndex = 0; FrameIndex < InputFrameCount; ++FrameIndex)
+				{
+					// Scale by 0.5f to preserve power when splitting into stereo
+					float InputSample = InputSamples[FrameIndex];
+					OutputSamples[SampleIndex++] = 0.5f * InputSample;
+					OutputSamples[SampleIndex++] = 0.5f * InputSample;
+				}
 			}
-		}
-		else
-		{
-			check(AudioDevice->SpatializeProcessor);
-
-			// Check if for OVR audio context initialization. We do this here because there apparently isn't a way to get
-			// the effect buffer size until the actual callback. 
-			if (!AudioDevice->SpatializeProcessor->IsSpatializationEffectInitialized())
+			else
 			{
-				AudioDevice->SpatializeProcessor->InitializeSpatializationEffect(InputFrameCount);
+				check(AudioDevice->SpatializeProcessor);
+
+				// Check if for OVR audio context initialization. We do this here because there apparently isn't a way to get
+				// the effect buffer size until the actual callback. 
+				if (!AudioDevice->SpatializeProcessor->IsSpatializationEffectInitialized())
+				{
+					AudioDevice->SpatializeProcessor->InitializeSpatializationEffect(InputFrameCount);
+				}
+
+				// Spatialize the audio stream with the current algorithm
+				AudioDevice->SpatializeProcessor->ProcessSpatializationForVoice(
+					SpatializationAlgorithm, 
+					VoiceId, 
+					InputSamples, 
+					OutputSamples, 
+					EmitterPosition
+				);
 			}
-
-			// Spatialize the audio stream with the current algorithm
-			AudioDevice->SpatializeProcessor->ProcessSpatializationForVoice(
-				SpatializationAlgorithm, 
-				VoiceId, 
-				InputSamples, 
-				OutputSamples, 
-				EmitterPosition
-			);
+			pOutputProcessParameters[0].BufferFlags = XAPO_BUFFER_VALID;
+			pOutputProcessParameters[0].ValidFrameCount = InputFrameCount;
 		}
-		pOutputProcessParameters[0].BufferFlags = XAPO_BUFFER_VALID;
-		pOutputProcessParameters[0].ValidFrameCount = InputFrameCount;
-	}
-	break;
+		break;
 
-	case XAPO_BUFFER_SILENT:
-	{
-		pOutputProcessParameters[0].BufferFlags = XAPO_BUFFER_SILENT;
-		pOutputProcessParameters[0].ValidFrameCount = InputFrameCount;
-	}
-	break;
+		case XAPO_BUFFER_SILENT:
+		{
+			pOutputProcessParameters[0].BufferFlags = XAPO_BUFFER_SILENT;
+			pOutputProcessParameters[0].ValidFrameCount = InputFrameCount;
+		}
+		break;
 	}
 
-	EndProcess();
 }
 
 /** Define the registration properties for the radio effect. */
diff --git a/Engine/Plugins/Runtime/OculusAudio/Source/OculusAudio/Private/OculusAudioEffect.h b/Engine/Plugins/Runtime/OculusAudio/Source/OculusAudio/Private/OculusAudioEffect.h
--- a/Engine/Plugins/Runtime/OculusAudio/Source/OculusAudio/Private/OculusAudioEffect.h
+++ b/Engine/Plugins/Runtime/OculusAudio/Source/OculusAudio/Private/OculusAudioEffect.h
@@ -16,41 +16,22 @@
 
 #define AUDIO_HRTF_EFFECT_CLASS_ID __declspec( uuid( "{8E67E588-FFF5-4860-A323-5E89B325D5EF}" ) )
 
-class FAudioHRTFEffectParameters
-{
-public:
-	FAudioHRTFEffectParameters()
-		: EmitterPosition(0.0f)
-		, SpatializationAlgorithm(SPATIALIZATION_TYPE_FAST)
-	{}
-
-	FAudioHRTFEffectParameters(const FVector& InEmitterPosition)
-		: EmitterPosition(InEmitterPosition)
-		, SpatializationAlgorithm(SPATIALIZATION_TYPE_FAST)
-	{}
-
-	FAudioHRTFEffectParameters(const FVector& InEmitterPosition, ESpatializationEffectType InSpatializationAlgorithm)
-		: EmitterPosition(InEmitterPosition)
-		, SpatializationAlgorithm(InSpatializationAlgorithm)
-	{}
-
-	FVector						EmitterPosition;
-	ESpatializationEffectType	SpatializationAlgorithm;
-};
-
-
 
-class AUDIO_HRTF_EFFECT_CLASS_ID FXAudio2HRTFEffect : public CXAPOParametersBase
+class AUDIO_HRTF_EFFECT_CLASS_ID FXAudio2HRTFEffect : public CXAPOBase
 {
 public:
 	FXAudio2HRTFEffect(uint32 InVoiceId, FAudioDevice* InAudioDevice);
+	~FXAudio2HRTFEffect();
+
 	STDMETHOD(LockForProcess)(UINT32 InputLockedParameterCount, const XAPO_LOCKFORPROCESS_BUFFER_PARAMETERS *pInputLockedParameters, UINT32 OutputLockedParameterCount, const XAPO_LOCKFORPROCESS_BUFFER_PARAMETERS *pOutputLockedParameters);
 	STDMETHOD_(void, Process)(UINT32 InputProcessParameterCount, const XAPO_PROCESS_BUFFER_PARAMETERS *pInputProcessParameters, UINT32 OutputProcessParameterCount, XAPO_PROCESS_BUFFER_PARAMETERS *pOutputProcessParameters, BOOL IsEnabled);
 
+	// Override AddRef/Release because we will be doing our own lifetime management without COM
+	STDMETHOD_(ULONG, AddRef)() override { return 0; }
+	STDMETHOD_(ULONG, Release)() override { return 0; }
+
 private:
 
-	FAudioHRTFEffectParameters				HRTFEffectParameters[3];
-	FVector									EmitterPosition;
 	WAVEFORMATEX							WaveFormat;
 
 	/* Which voice index this effect is attached to. */
diff --git a/Engine/Source/Runtime/Engine/Public/IAudioExtensionPlugin.h b/Engine/Source/Runtime/Engine/Public/IAudioExtensionPlugin.h
--- a/Engine/Source/Runtime/Engine/Public/IAudioExtensionPlugin.h
+++ b/Engine/Source/Runtime/Engine/Public/IAudioExtensionPlugin.h
@@ -19,6 +19,28 @@ enum ESpatializationEffectType
 	SPATIALIZATION_TYPE_COUNT
 };
 
+class FAudioSpatializationParams
+{
+public:
+	FAudioSpatializationParams()
+		: EmitterPosition(0.0f)
+		, SpatializationAlgorithm(SPATIALIZATION_TYPE_FAST)
+	{}
+
+	FAudioSpatializationParams(const FVector& InEmitterPosition)
+		: EmitterPosition(InEmitterPosition)
+		, SpatializationAlgorithm(SPATIALIZATION_TYPE_FAST)
+	{}
+
+	FAudioSpatializationParams(const FVector& InEmitterPosition, ESpatializationEffectType InSpatializationAlgorithm)
+		: EmitterPosition(InEmitterPosition)
+		, SpatializationAlgorithm(InSpatializationAlgorithm)
+	{}
+
+	FVector						EmitterPosition;
+	ESpatializationEffectType	SpatializationAlgorithm;
+};
+
 /**
 * IAudioSpatializationAlgorithm
 *
@@ -41,9 +63,14 @@ public:
 	{
 
 	}
+
 	/** Sets the spatialization effect parameters. */
+	virtual void SetSpatializationParameters(uint32 VoiceId, const FAudioSpatializationParams& Params)
+	{
+	}
 
-	virtual void SetSpatializationParameters(uint32 VoiceId, const FVector& EmitterPosition, ESpatializationEffectType AlgorithmType)
+	/** Gets the spatialization effect parameters. */
+	virtual void GetSpatializationParameters(uint32 VoiceId, FAudioSpatializationParams& OutParams)
 	{
 	}
 
diff --git a/Engine/Source/Runtime/Windows/XAudio2/Private/XAudio2Source.cpp b/Engine/Source/Runtime/Windows/XAudio2/Private/XAudio2Source.cpp
--- a/Engine/Source/Runtime/Windows/XAudio2/Private/XAudio2Source.cpp
+++ b/Engine/Source/Runtime/Windows/XAudio2/Private/XAudio2Source.cpp
@@ -648,7 +648,7 @@ void FXAudio2SoundSource::GetMonoChannelVolumes(float ChannelVolumes[CHANNEL_MAT
 		FVector UnnormalizedDirection = AudioDevice->InverseListenerTransform.TransformPosition(WaveInstance->Location);
 		EmitterPosition = UnnormalizedDirection.GetSafeNormal();
 
-		AudioDevice->SpatializeProcessor->SetSpatializationParameters(VoiceId, EmitterPosition, (ESpatializationEffectType)WaveInstance->SpatializationAlgorithm);
+		AudioDevice->SpatializeProcessor->SetSpatializationParameters(VoiceId, FAudioSpatializationParams(EmitterPosition, (ESpatializationEffectType)WaveInstance->SpatializationAlgorithm));
 		GetStereoChannelVolumes(ChannelVolumes, AttenuatedVolume);
 	}
 	else // Spatialize the mono stream using the normal 3d audio algorithm
