Copying //UE4/Dev-AnimPhys to //UE4/Dev-Main (Source: //UE4/Dev-AnimPhys @ 3624379)

============================
  MAJOR FEATURES & CHANGES
============================

Change 3536809 by [REDACTED_PERSON]

	Fixing case of files in "iOS" directory, pt 1.

Change 3536814 by [REDACTED_PERSON]

	Fixing case of files in "iOS" directory, pt 2.

Change 3596207 by [REDACTED_PERSON]

	Copying //Tasks/UE4/Dev-UEAP-29-PhATUpgrade to Dev-AnimPhys (//UE4/Dev-AnimPhys) @ CL 3590250

	PhAT Upgrade

	Changelists from task stream:

	Change 3380649 on 2017/04/05 by [REDACTED_PERSON][REDACTED_HANDLE]

		Initial pass at allowing viewports to be extended more easily, still plenty TOD, but just unearthing this old shelf and getting it working.

		This gets the Persona skeleton tree and viewport into PhAT, without any PhAT functionality added.

	Change 3380685 on 2017/04/05 by [REDACTED_PERSON][REDACTED_HANDLE]

		Renaming PhAT files to PhysicsAssetEditor

	Change 3380749 on 2017/04/05 by [REDACTED_PERSON][REDACTED_HANDLE]

		Rename PhAT -> PhysicsAssetEditor

	Change 3380832 on 2017/04/05 by [REDACTED_PERSON][REDACTED_HANDLE]

		Fixed up PhAT to Physics Asset Editor

	Change 3380884 on 2017/04/05 by [REDACTED_PERSON][REDACTED_HANDLE]

		Reverted some over-zealous renaming

	Change 3380970 on 2017/04/05 by [REDACTED_PERSON][REDACTED_HANDLE]

		Tweaked ISkeletonTreeBuilder interface to make way for actually making a derived class of it
		Added the ability to hide filter menus to skeleton tree

	Change 3381017 on 2017/04/05 by [REDACTED_PERSON][REDACTED_HANDLE]

		Added new physics asset skeleton tree builder

	Change 3384407 on 2017/04/07 by [REDACTED_PERSON][REDACTED_HANDLE]

		Skeleton tree extensions to support physics assets

		Only started this work - still much to do

	Change 3384460 on 2017/04/07 by [REDACTED_PERSON][REDACTED_HANDLE]

		Rearranged persona viewport menus

	Change 3392222 on 2017/04/13 by [REDACTED_PERSON][REDACTED_HANDLE]

		Removed body/constraint modes.
		Added graph editor
		Added edit mode - moved viewport client code over
		Got PhAT skel mesh rendering in viewport

	Change 3392268 on 2017/04/13 by [REDACTED_PERSON][REDACTED_HANDLE]

		Increased hit proxy priority to improve selection

	Change 3401648 on 2017/04/20 by [REDACTED_PERSON][REDACTED_HANDLE]

		Skeleton tree gets bodies & shapes back.

		Selection works in graph, now displaying the correct constraint in the detials panel. Still need to add selection from viewport.

		Added multi-select to bone proxy customization

		Re-tweaked editor layout

	Change 3403701 on 2017/04/21 by [REDACTED_PERSON][REDACTED_HANDLE]

		Selection sync work.

		Customization of anim viewport menus.

		Context menus for physics asset items, as well as masking of various context menu items via settings.

	Change 3405246 on 2017/04/24 by [REDACTED_PERSON][REDACTED_HANDLE]

		Started more work on viewport menu extensions, but need to refactor the toolbar system to use actual multiboxes. Up next!

	Change 3405274 on 2017/04/24 by [REDACTED_PERSON][REDACTED_HANDLE]

		More viewport menu fixups (plus deleting duplicate functionality).

	Change 3409155 on 2017/04/26 by [REDACTED_PERSON][REDACTED_HANDLE]

		Got simulation working again - as we switched to the debug skel mesh comp, the normal tick path didnt work for post-blend physics (it tried to flip the buffer too early).
		Also tweaked debug skel mesh comp root motion consumption code to not reset transfor every frame if we are not using root motion.
		Cleaned up unused files & code

	Change 3410814 on 2017/04/27 by [REDACTED_PERSON][REDACTED_HANDLE]

		Allow extensibility of viewport menu bars

		Slate changes:

		Allow menu bars to optionally specify an icon to use. This is intended to allow us to move viewport tool/menu bars over to use multibox, with all the attendant features and extension points.
		Allow menu bars to optionally invert-on-hover.
		Allow styling of menus to affect closed appearance of menu header. Previously only NoBorder was used.
		Adjusted core styling of menu bar elements.

		Other changes:

		Adjusted padding for various UI elements to preserve previoud behavior.
		Adjusted SAnimViewportToolbar to use the new menu bar builder.
		Exposed SEditorViewportViewMenu so that it can be used in a standard menu bar.

	Change 3410816 on 2017/04/27 by [REDACTED_PERSON][REDACTED_HANDLE]

		Added extension point to viewport menu bar

	Change 3410818 on 2017/04/27 by [REDACTED_PERSON][REDACTED_HANDLE]

		Getting sim working again

		Moved over to using preview instance so we share functionality with Persona editors.
		Added time dilation options to persona preview scene.
		Removed PhAT specific recording functionality (it is in the viewport now).

	Change 3410840 on 2017/04/27 by [REDACTED_PERSON][REDACTED_HANDLE]

		Recreate physics state on edit, not sim start

		This allows velocity to be inherited when simulation is started

	Change 3410863 on 2017/04/27 by [REDACTED_PERSON][REDACTED_HANDLE]

		Moving viewport to continually-invalidated one like animation editors

		Fixed crash in non-extended viewport toolbars

	Change 3410936 on 2017/04/27 by [REDACTED_PERSON][REDACTED_HANDLE]

		Bodies start off non-expanded

		Selection now synced between viewport and graph

		Constraint selection in graph not works on the first try

	Change 3410943 on 2017/04/27 by [REDACTED_PERSON][REDACTED_HANDLE]

		Added missing icon

	Change 3410966 on 2017/04/27 by [REDACTED_PERSON][REDACTED_HANDLE]

		Removed shape listing from graph nodes

	Change 3411013 on 2017/04/27 by [REDACTED_PERSON][REDACTED_HANDLE]

		Double click on body node recenters graph

		Fixed graph disappearing on right-click

	Change 3411111 on 2017/04/27 by [REDACTED_PERSON][REDACTED_HANDLE]

		Prevented cursor getting swallowed in sim mode

	Change 3411126 on 2017/04/27 by [REDACTED_PERSON][REDACTED_HANDLE]

		Fixed overlapping text

	Change 3411213 on 2017/04/27 by [REDACTED_PERSON][REDACTED_HANDLE]

		Node layout now takes dimensions into account

	Change 3411320 on 2017/04/27 by [REDACTED_PERSON][REDACTED_HANDLE]

		Fixed crash opening Persona editors

		Renamed file

	Change 3411327 on 2017/04/27 by [REDACTED_PERSON][REDACTED_HANDLE]

		Tweaks to profiles menu

	Change 3420822 on 2017/05/03 by [REDACTED_PERSON][REDACTED_HANDLE]

		Profiles can now be edited in their own details panel

		Existing customizations folded into the new panel

		Tweaks to toolbar

		Added the ability for the persona details panel to have extra top/bottom content added

	Change 3420832 on 2017/05/03 by [REDACTED_PERSON][REDACTED_HANDLE]

		Add profile control to context menus

		Also delete old unused code

	Change 3422651 on 2017/05/04 by [REDACTED_PERSON][REDACTED_HANDLE]

		Toolbar trimmed down & re-ordered

		Body/constraint ops moved to context menus

		Apply physmat now a context-menu option with an asset picker

	Change 3422654 on 2017/05/04 by [REDACTED_PERSON][REDACTED_HANDLE]

		Removed extra warning dialog when auto-creating bodies

		Changed title of new asset dialog to "auto-create bodies"

	Change 3422680 on 2017/05/04 by [REDACTED_PERSON][REDACTED_HANDLE]

		Fix "simulate selected"

		As we dont re-init the physics state each time we start simulating, our tweaked physics type was never applied. We now manually do this in the editor.

	Change 3422937 on 2017/05/04 by [REDACTED_PERSON][REDACTED_HANDLE]

		Replaced EKCollisionPrimitiveType with EAggCollisionShape::Type

		Fixed up selection so body selection works & tree seleciton is properly synced with viewport
		Added recursion guard to selection delegate handlers.

		Removed vestigial instance property editing support (no longer needed).

		Removed unused old tree support code

	Change 3423034 on 2017/05/04 by [REDACTED_PERSON][REDACTED_HANDLE]

		Added constraints to tree

	Change 3423318 on 2017/05/04 by [REDACTED_PERSON][REDACTED_HANDLE]

		Fix bone proxiies stopping updating after initial viewport selection

	Change 3424993 on 2017/05/05 by [REDACTED_PERSON][REDACTED_HANDLE]

		Fixed up selection issues when creating new bodies

		Added constraint context menu

	Change 3424998 on 2017/05/05 by [REDACTED_PERSON][REDACTED_HANDLE]

		Moved icons to central location

	Change 3425445 on 2017/05/05 by [REDACTED_PERSON][REDACTED_HANDLE]

		Customized filtering of the skeleton tree

		Hide constraints by defualt

		Added option to hide parents when filtering (so the vertical space is nto wasted, but some idea of hierarchy is preserved).

		BREAKING CHANGE: changed skeleton tree filtering API to add args & removed bWillFilter bool.

	Change 3425488 on 2017/05/05 by [REDACTED_PERSON][REDACTED_HANDLE]

		Merging //UE4/Dev-AnimPhys to Dev-UEAP-29-PhATUpgrade (//Tasks/UE4/Dev-UEAP-29-PhATUpgrade) @ CL 3425303

	Change 3427886 on 2017/05/08 by [REDACTED_PERSON][REDACTED_HANDLE]

		Moved physics sim options to viewport menu (so seleciton changing is not required to change them)

		Moved physics-related rendering options to show menu

		We no longer switch to sim options when nothing is selected.

		During simulation we now disable the details panel

		Constraint scaling now works correctly (rather than just scaling the screen size limit that axes only are rendered)

	Change 3428040 on 2017/05/08 by [REDACTED_PERSON][REDACTED_HANDLE]

		Small fixes based on feedback:

		Exposed Mirror tool to menus

		Exposed constraint quick actions to menus

		Added edit condition to Position & Velocity strength for physical animation

		Fixed up some tooltips & display names

	Change 3428143 on 2017/05/08 by [REDACTED_PERSON][REDACTED_HANDLE]

		Defaulted to constraints as points

	Change 3428216 on 2017/05/08 by [REDACTED_PERSON][REDACTED_HANDLE]

		Request from Nick D:

		Update in-level primitive transforms immediately, rather than on mouse up. We only do this for non-convex primitives however, to avoid re-cooking meshes.

	Change 3430326 on 2017/05/09 by [REDACTED_PERSON][REDACTED_HANDLE]

		Tweaks to rendering of constraints and shapes to allow for better seleciton & interaction with editor widgets.

		Slightly increased point-constraint rendering size and added crosshair cursor to constraints

	Change 3430327 on 2017/05/09 by [REDACTED_PERSON][REDACTED_HANDLE]

		Fixed object-reuse issue in skeleton tree items with sanem names (use a GUID instead)

	Change 3430391 on 2017/05/09 by [REDACTED_PERSON][REDACTED_HANDLE]

		Removed duplicate time dilation (can just use viewport menu!)

	Change 3430419 on 2017/05/09 by [REDACTED_PERSON][REDACTED_HANDLE]

		Fixup post-merge

		Prevent crash by attaching to root component in the correct place

		Add IWYU include for TArrayView

		Remove more unused code

	Change 3430443 on 2017/05/09 by [REDACTED_PERSON][REDACTED_HANDLE]

		Fix constraint/body selection one final time

		Move constraint drawing to SDPG_World (apart from point mode)

		Remove depth offset in material

	Change 3430495 on 2017/05/09 by [REDACTED_PERSON][REDACTED_HANDLE]

		Enabling/disabling collision between bodies is now clearer

		Menu items are now enabled and disabled correctly depending on collision state

		Tooltip reflects what actually gets done when the operation is enacted

		Also corrected a few functions that still reference constraint & body mode

	Change 3430553 on 2017/05/09 by [REDACTED_PERSON][REDACTED_HANDLE]

		Added enable/disable collision with all

	Change 3432386 on 2017/05/10 by [REDACTED_PERSON][REDACTED_HANDLE]

		Color code graph items based on current profile

	Change 3432401 on 2017/05/10 by [REDACTED_PERSON][REDACTED_HANDLE]

		Color code tree items too

	Change 3432418 on 2017/05/10 by [REDACTED_PERSON][REDACTED_HANDLE]

		Bone selection & manipulation now possible - allows for pose setup before simulation

		Item expansion now expands leaf nodes when selecting - helps with constraint selection etc.

	Change 3432427 on 2017/05/10 by [REDACTED_PERSON][REDACTED_HANDLE]

		Fix compile error

		Color code according to simulated/kinematic status

	Change 3432428 on 2017/05/10 by [REDACTED_PERSON][REDACTED_HANDLE]

		File i missed

	Change 3432540 on 2017/05/10 by [REDACTED_PERSON][REDACTED_HANDLE]

		Added physics asset factory so physics assets can be created form the "new asset" menu.

		Skeletal mesh is picked then a defualt asset is generated

	Change 3432556 on 2017/05/10 by [REDACTED_PERSON][REDACTED_HANDLE]

		Improve interactions with bones & bodies

		Clear bone selection when selecting bodies/constraints

		Always hide gizmo in simulate

	Change 3432703 on 2017/05/10 by [REDACTED_PERSON][REDACTED_HANDLE]

		Removed unused selection lock feature

		Fixed selection working incorrectly with details panel closed

	Change 3434710 on 2017/05/11 by [REDACTED_PERSON][REDACTED_HANDLE]

		Selection improvements

		Multiselect in tree now only selects non-collapsed tree elements

		Selection API revamped in shared data, so multiselect of constraints can work correctly (they appear more than once in the tree, so the preivous single-point-of-access API was insufficent).

	Change 3489030 on 2017/06/14 by [REDACTED_PERSON][REDACTED_HANDLE]

		Merging //UE4/Dev-AnimPhys to Dev-UEAP-29-PhATUpgrade (//Tasks/UE4/Dev-UEAP-29-PhATUpgrade) @ CL 3488994

	Change 3491459 on 2017/06/15 by [REDACTED_PERSON][REDACTED_HANDLE]

		Fixup post-merge issues

	Change 3491486 on 2017/06/15 by [REDACTED_PERSON][REDACTED_HANDLE]

		Simulation now works in a simlar way to the level editor

		Only on 'simulate' button, which controls repeating the last simulation (be it selected or not). Options are on a dropdown.

	Change 3491529 on 2017/06/15 by [REDACTED_PERSON][REDACTED_HANDLE]

		Fixed selection color of wireframe drawing (this broke ages ago!)
		Fixed initialized environment color/intensity

	Change 3491537 on 2017/06/15 by [REDACTED_PERSON][REDACTED_HANDLE]

		Tweaked materials so they dont repend on seperate translucency (which is optional, and disabled currently)

	Change 3491791 on 2017/06/15 by [REDACTED_PERSON][REDACTED_HANDLE]

		Fix crash when simulating selected new bodies

		Make sure we recreate physics state appropriately (it used to be done on simulation start, so wasnt needed each time)

	Change 3494359 on 2017/06/16 by [REDACTED_PERSON][REDACTED_HANDLE]

		Select all is now a menu option

		Context menu pops when right-clicking nothing now too
		Menu no longer grows enormous when multiple types of objects are selected

	Change 3494373 on 2017/06/16 by [REDACTED_PERSON][REDACTED_HANDLE]

		Enlarged constraint rendering size
		Show constraints (rather than points) by default

	Change 3511708 on 2017/06/27 by [REDACTED_PERSON][REDACTED_HANDLE]

		Physics Assets now appear in the asset family shortcut bar

		Physics Assets now render thumbnails

		Skeleton tree can now work in 'picker' mode

		Constraints can now be created manually in the graph, tree and viewport

		Fixed double-click and mousewheel not working right sometimes

	Change 3513121 on 2017/06/28 by [REDACTED_PERSON][REDACTED_HANDLE]

		Fixed clicks incorrectly selecting bones in simulate mode

	Change 3513160 on 2017/06/28 by [REDACTED_PERSON][REDACTED_HANDLE]

		Physics Asset config is now loaded/saved

		Fixed antoher corner case with viewport clicks in sim

	Change 3513540 on 2017/06/28 by [REDACTED_PERSON][REDACTED_HANDLE]

		Moved body creation params over to a details panel & settings object

		Moved initial creation dialog over to use the new system too

	Change 3513591 on 2017/06/28 by [REDACTED_PERSON][REDACTED_HANDLE]

		Renamed shapes and constraints in the tree view

	Change 3513752 on 2017/06/28 by [REDACTED_PERSON][REDACTED_HANDLE]

		Constraints are now not filtered by default

	Change 3513797 on 2017/06/28 by [REDACTED_PERSON][REDACTED_HANDLE]

		Selecting constraints now shows them (and the bodies involved) in the graph

	Change 3513859 on 2017/06/28 by [REDACTED_PERSON][REDACTED_HANDLE]

		Removed "Show Kinematic Bodies"

		We now always show kinematic status in simulate mode

	Change 3515732 on 2017/06/29 by [REDACTED_PERSON][REDACTED_HANDLE]

		PhAT rendering settings are now persisted across sessions.

		Access to sim/edit settings is now not gated on state of the editor. Sim/edit settings are always both available.

		Added editable opacity to collision rendering.

	Change 3515735 on 2017/06/29 by [REDACTED_PERSON][REDACTED_HANDLE]

		New materials with opacity parameter

	Change 3515757 on 2017/06/29 by [REDACTED_EMAIL]

		Re-saved materials

	Change 3515759 on 2017/06/29 by [REDACTED_PERSON][REDACTED_HANDLE]

		Added ability to only show selected bodies as solid

	Change 3515812 on 2017/06/29 by [REDACTED_PERSON][REDACTED_HANDLE]

		Fix focus 'F' shortcut sometimes not working

	Change 3515984 on 2017/06/29 by [REDACTED_PERSON][REDACTED_HANDLE]

		Fix a bunch of selection issues with the graph not keeping in sync

	Change 3517456 on 2017/06/30 by [REDACTED_PERSON][REDACTED_HANDLE]

		Merging //UE4/Dev-AnimPhys to Dev-UEAP-29-PhATUpgrade (//Tasks/UE4/Dev-UEAP-29-PhATUpgrade) @ CL 3516853

	Change 3517514 on 2017/06/30 by [REDACTED_PERSON][REDACTED_HANDLE]

		Fixed disappearing convex meshes on simulate

		Also fixes crash in thumbnail rendering

	Change 3517556 on 2017/06/30 by [REDACTED_PERSON][REDACTED_HANDLE]

		Disabled selection on mesh. Fixes selection issues.

		Also made the hit proxy use a crosshair when over bodies, for easier selection

	Change 3517642 on 2017/06/30 by [REDACTED_PERSON][REDACTED_HANDLE]

		Added body/body collision buttons back to the main toolbar

		Fixed solid body drawing using the wrong material when no bodies are selected

	Change 3517828 on 2017/06/30 by [REDACTED_PERSON][REDACTED_HANDLE]

		Fix delete shortcut not working when tree is focused

	Change 3517927 on 2017/06/30 by [REDACTED_PERSON][REDACTED_HANDLE]

		Integrated per-bone primitive generation with the new tab method

		Removed context menu item for bones (fixes duplicate popup)

		Fixed undo/redo not working for regenerating all bodies

	Change 3519931 on 2017/07/03 by [REDACTED_PERSON][REDACTED_HANDLE]

		Disabled body regeneration when simulation is running

		Fixed up tab icons

	Change 3519978 on 2017/07/03 by [REDACTED_PERSON][REDACTED_HANDLE]

		Preview mesh is now set like every other Persona editor (via toolbar picker of via preview scene settings)

		Animation picker removed from toolbar (we use the preview scene settings for this now)

		Fixed profiles tab icon

	Change 3519982 on 2017/07/03 by [REDACTED_PERSON][REDACTED_HANDLE]

		Show attached assets in tree

	Change 3519995 on 2017/07/03 by [REDACTED_PERSON][REDACTED_HANDLE]

		Fix broken multi-selection of bone proxies

	Change 3532799 on 2017/07/12 by [REDACTED_PERSON][REDACTED_HANDLE]

		Removed code that prevented parts of the UI (like simulation) from working in PIE

		Removed graph overlays & added "PHYSICS" label

	Change 3532837 on 2017/07/12 by [REDACTED_PERSON][REDACTED_HANDLE]

		Removed arrows from graph

		Fixed dragging off constraints/input pins/bodies in constraint-created graphs

		Constraint names now include both bodies

	Change 3532880 on 2017/07/12 by [REDACTED_PERSON][REDACTED_HANDLE]

		Switched from colors to icons in the skeleton tree

		Removed bold fonts

	Change 3532907 on 2017/07/12 by [REDACTED_PERSON][REDACTED_HANDLE]

		Layout fixes

		Added border around generate button in tools panel

		Removed skeleton tree header in contexts where it is not needed

	Change 3532932 on 2017/07/12 by [REDACTED_PERSON][REDACTED_HANDLE]

		Added slow task dialog for body generation

	Change 3532992 on 2017/07/12 by [REDACTED_PERSON][REDACTED_HANDLE]

		Rearranged context menus to be not so huge

	Change 3533134 on 2017/07/12 by [REDACTED_PERSON][REDACTED_HANDLE]

		Rearranged menus some more

	Change 3533135 on 2017/07/12 by [REDACTED_PERSON][REDACTED_HANDLE]

		Colorized details customization of swing/twist items

	Change 3533174 on 2017/07/12 by [REDACTED_PERSON][REDACTED_HANDLE]

		Auto-open assets when creating from skeletal mesh

		Tweaked tooltip on suggestion from Nick D

	Change 3535652 on 2017/07/13 by [REDACTED_PERSON][REDACTED_HANDLE]

		Fixed mirroring changes not showing up straight away

	Change 3535731 on 2017/07/13 by [REDACTED_PERSON][REDACTED_HANDLE]

		Moved over to Persona-style floor adjustment

	Change 3539689 on 2017/07/17 by [REDACTED_PERSON][REDACTED_HANDLE]

		Tweaked tooltips for filtering items

	Change 3539693 on 2017/07/17 by [REDACTED_PERSON][REDACTED_HANDLE]

		Added "deselect all" option (Esc)

	Change 3539731 on 2017/07/17 by [REDACTED_PERSON][REDACTED_HANDLE]

		Graph selection tweaks

		Selected bodies in the viewport/tree are now also selected in the graph.
		Selection outline is now matched to the graph outline instead of using default outline.
		Pin allocation no longer happens twice

	Change 3539750 on 2017/07/17 by [REDACTED_PERSON][REDACTED_HANDLE]

		Switched simulate shortcut to Alt+Enter

		Avoids conflict with clobal PIS/SIE shortcuts

	Change 3539933 on 2017/07/17 by [REDACTED_PERSON][REDACTED_HANDLE]

		Minor body regeneration refactor

		Label for tools tab button is dynamic depending on selection context
		Generation setttings are now re-used by creation dialog too
		Added in per-bone and per-body regeneration menu items. Bone regeneration now deletes the old body(s) instead of aborting

	Change 3543884 on 2017/07/19 by [REDACTED_PERSON][REDACTED_HANDLE]

		Resetting animation to default now correctly applies the animation

	Change 3544101 on 2017/07/19 by [REDACTED_PERSON][REDACTED_HANDLE]

		Fixed up physics asset editor's use of debug skel mesh component

		This broke post-merge from Dev-AnimPhys.
		Kinda hacky, but we need to double-flip the buffers in this case as we want to force non-threaded work AND also wait on the physics tick group to complete (to blend in physics).
		This also requires making ShouldBlendPhysicsBones protected, otherwise the buffers are never flipped in the non=simulating case (before simulation is enabled in the physics asset editor).

	Change 3547893 on 2017/07/21 by [REDACTED_PERSON][REDACTED_HANDLE]

		Moved code to add/remove/assign/unassign profiles to details customization

		Also allowed dupication again (via the menu)
		Allows correct naming of new profiles as before (as this is handled in PostEdit)

	Change 3547925 on 2017/07/21 by [REDACTED_PERSON][REDACTED_HANDLE]

		Prevented ctrl+selection of constraints from re-selecting

		Avoided defered broadcast of seleciton event from the graph

	Change 3550332 on 2017/07/24 by [REDACTED_PERSON][REDACTED_HANDLE]

		Fixed bodies incorrectly simulating outside of 'simulate' mode

		Forced all bodies to be non-simulated when simulation is disabled.

		Also removed non-functioning motor menu options & disabled more menu options when simulating

	Change 3550355 on 2017/07/24 by [REDACTED_PERSON][REDACTED_HANDLE]

		Fixed crash when failing to create a physics asset with multi convex hull

	Change 3558007 on 2017/07/27 by [REDACTED_PERSON][REDACTED_HANDLE]

		Fixed typo that disabled editability of profile names incorrectly

	Change 3566157 on 2017/08/01 by [REDACTED_PERSON][REDACTED_HANDLE]

		Fixed crash when opening a physics asset with a deleted preview skeletal mesh

		Now assigns default mesh as before

		If the mesh is then reset, the asset editor must be re-opened as the skeleton will have changed underneath it.

	Change 3568327 on 2017/08/02 by [REDACTED_PERSON][REDACTED_HANDLE]

		Prevent "set bodies below" from improperly enabling simulation on bodies

	Change 3570436 on 2017/08/03 by [REDACTED_PERSON][REDACTED_HANDLE]

		Physics assets with simulated bodies no longer simulate when first opened

	Change 3570470 on 2017/08/03 by [REDACTED_PERSON][REDACTED_HANDLE]

		Fix excessive gravity crash when actors pop out of the world

		Also restrict gravity to non NaN-causing levels.

	Change 3570717 on 2017/08/03 by [REDACTED_PERSON][REDACTED_HANDLE]

		Merging //UE4/Dev-AnimPhys to Dev-UEAP-29-PhATUpgrade (//Tasks/UE4/Dev-UEAP-29-PhATUpgrade) @ CL 3570581

	Change 3570781 on 2017/08/03 by [REDACTED_PERSON][REDACTED_HANDLE]

		Fix merge issues

	Change 3587760 on 2017/08/15 by [REDACTED_PERSON][REDACTED_HANDLE]

		Removed delegate for skeleton tree context menu extension, now uses an empty section

	Change 3589915 on 2017/08/16 by [REDACTED_PERSON][REDACTED_HANDLE]

		Added comments to bone proxy & physics asset editor shared data

		Removed unused variables

	Change 3589976 on 2017/08/16 by [REDACTED_PERSON][REDACTED_HANDLE]

		Fixed constraint 'all positions' rendering

		Removed empty override of unregister tab spawners

	Change 3589983 on 2017/08/16 by [REDACTED_PERSON][REDACTED_HANDLE]

		Fix crash when setting skeletal mesh

		Toast is not displayed when the skeleton is changed as well as the skeletal mesh.
		Toolkit was getting invalidated as setting the preview mesh to a different skeleton ends up restarting the sub-editor

	Change 3589990 on 2017/08/16 by [REDACTED_PERSON][REDACTED_HANDLE]

		Physics asset selection color now uses editor settings

	Change 3589994 on 2017/08/16 by [REDACTED_PERSON][REDACTED_HANDLE]

		Removed unused functions

	Change 3589997 on 2017/08/16 by [REDACTED_PERSON][REDACTED_HANDLE]

		Commented SetBodiesBelowPhysicsType as per code review

	Change 3590007 on 2017/08/16 by [REDACTED_PERSON][REDACTED_HANDLE]

		Disabled physical material menu in simulate

	Change 3590130 on 2017/08/16 by [REDACTED_PERSON][REDACTED_HANDLE]

		Removed unused code

		Commented a few functions

		Re-instated preview mesh selection

		Removed delegate allowing viewport client class creation

	Change 3590154 on 2017/08/16 by [REDACTED_PERSON][REDACTED_HANDLE]

		Remove unused code

	Change 3590197 on 2017/08/16 by [REDACTED_PERSON][REDACTED_HANDLE]

		Merging //UE4/Dev-AnimPhys to Dev-UEAP-29-PhATUpgrade (//Tasks/UE4/Dev-UEAP-29-PhATUpgrade) @ CL 3589965

	Change 3590250 on 2017/08/16 by [REDACTED_PERSON][REDACTED_HANDLE]

		Fixup merge errors

Change 3596227 by [REDACTED_PERSON]

	Fixed physics substepping interpolation using the wrong starting value.

Change 3596241 by [REDACTED_PERSON]

	Fixed cloth not being drawn correctly in the editor, due to bounds not being computed accurately.

Change 3596247 by [REDACTED_PERSON]

	Fixup CIS errors post PhAT Upgrade merge

Change 3596250 by [REDACTED_PERSON]

	More CIS fixes

Change 3596255 by [REDACTED_PERSON]

	Fixed compilation errors when nativizing animation blueprints that use subinstances

Change 3596256 by [REDACTED_PERSON]

	Fixed orphaned sub anim instance pins hanging around

Change 3596257 by [REDACTED_PERSON]

	Fixed skel surf particles being misplaced when clothing was active. And fixed particles spawning on disabled cloth proxy sections.

Change 3596258 by [REDACTED_PERSON]

	Hide mass override when selecting skeletal meshes. Mass overrides are taken from physics asset and will be ignored on the component so it makes no sense to have this visible

Change 3596259 by [REDACTED_PERSON]

	Fixed mismatch between paint values and view values for clothing tools

Change 3596260 by [REDACTED_PERSON]

	Stopped property context menus killing the whole window stack when an item is clicked

Change 3596261 by [REDACTED_PERSON]

	One last Mac CIS fix (hopefully)

Change 3596308 by [REDACTED_PERSON]

	Removed outdated references to APEX in clothing example map.

Change 3596360 by [REDACTED_PERSON]

	Fixing inconsistent animation entries in blueprint context menu (displaying differently depending on whether the asset is loaded) + Cache correct tooltip when asset isn't loaded

Change 3596459 by [REDACTED_PERSON]

	Fixed anim curves not correctly being updated to post process instances. Change made to curve update in Dev-General fixed main and sub instances but missed post process instances.

Change 3596967 by [REDACTED_PERSON]

	Adding setting default reverb send level in audio settings.

Change 3596974 by [REDACTED_PERSON]

	Merge in fix from Christopher Oliver

Change 3597243 by [REDACTED_PERSON]

	Checking in missing files.

Change 3597686 by [REDACTED_PERSON]

	Fix warnings/errors from CL 3597452

Change 3597846 by [REDACTED_PERSON]

	Fix errors, take 2

Change 3598290 by [REDACTED_PERSON]

	Panning Angle Issue

Change 3598412 by [REDACTED_PERSON]

	Change Core.h header to CoreMinimal.h, fix warnings

Change 3599797 by Jurre.deBaare

	LODs from Merge Actor tool have bad normals
	#fix normals weren't wrong but user was complaining about the lightmap resolution behaviour, so added a new feature that calculates the lightmap resolution according to:
		1) Summing all lightmap pixel counts for each mesh component being merged
		2) Calculating fitting texture dimension by taking square root of the total pixels

Change 3599863 by [REDACTED_PERSON]

	PR #3919: rename flag 'DEPERCATED_PHYSBLEND_UPDATES_PHYSX' to 'DEPRECATED_PHYSBLEND_UPDATES_PHYSX' to fix the typo (Contributed by [REDACTED_PERSON])

Change 3599883 by Jurre.deBaare

	HLOD: update outliner tooltip when UE docs arrive

Change 3599944 by [REDACTED_PERSON]

	Smart name refactor
	- Remove guids entirely
	- Remove automatic fix up
	- Simplify smart name mapping container
	- Make animations deterministic for cooking

Change 3600133 by [REDACTED_PERSON]

	Fixed crash shutting down editor with active cloth paint tab, as mode manager was being used unsafely.

Change 3600166 by [REDACTED_PERSON]

	Fixed cloth paint gradient allowing invalid values

Change 3600719 by [REDACTED_PERSON]

	PR #3894: PlayMontage node bug Fix (Contributed by [REDACTED_PERSON])

Change 3601668 by Jurre.deBaare

	Improve BlendSpace preview pin dragging controls
	#fix Click and drag now also works for the preview pin which should allign it with other pins on the grid and makes the preview functionality more discoverable
	#misc Also added tooltips on the grid to make the functionality more discoverable

Change 3601669 by Jurre.deBaare

	No easy way to tell which Blend Sample in the blend graph matches up to which Blend Sample in the Asset Details panel
	#fix I've added the SampleIndex to the names to make it easier recognizing which one is which

Change 3601731 by [REDACTED_PERSON]

	Fixed cloth paint falloff to actually calculate falloff, and take brush strength into account.

Change 3601897 by [REDACTED_PERSON]

	fixing issue with sequencer reinitialization

Change 3602339 by [REDACTED_PERSON]

	Fixed comment/tooltip typo

Change 3602502 by [REDACTED_PERSON]

	Fixed clothing gradient tool renderer not showing selected points when camera was moving

Change 3602664 by [REDACTED_PERSON]

	Unshelved fixes from Dev-VR

Change 3602726 by [REDACTED_PERSON]

	Back out revision 3 from //UE4/Dev-AnimPhys/QAGame/QAGame.uproject

Change 3603011 by [REDACTED_PERSON]

	Fix build error

Change 3604139 by [REDACTED_PERSON]

	Restricted painter processing to no longer attempt painting while in simulation previews in cloth paint mode.

Change 3604284 by [REDACTED_PERSON]

	Fixed crashes in physics asset editor and skeletal mesh editor when the preview scene clears out the preview mesh while clothing is running

Change 3604612 by [REDACTED_PERSON]

	Fix curve issue from automation test
	- It was actual bug.

Change 3604614 by [REDACTED_PERSON]

	- Fix crash with macro anim notify
	- Make sure macro anim notify doesn't show up in the menu

Change 3604725 by [REDACTED_PERSON]

	fixed issue with opening state machine from anim graph

Change 3604971 by [REDACTED_PERSON]

	Bringing fix from 4.17 to Dev-AnimPhys

Change 3605787 by [REDACTED_PERSON]

	Adding ability to pass in an optional owner in PlaySound2D and PlaySoundAtLocation BP calls

	- This is necessary in order to use the sound concurrency "limit by owner" feature

Change 3606851 by Jurre.deBaare

	UE4Editor Static Analysis Win64 - Warning fix

Change 3607022 by [REDACTED_PERSON]

	Fix static analysis warning

Change 3607229 by Jurre.deBaare

	RemoveAllCurveData should not allow removing data from the Skeleton

Change 3607660 by [REDACTED_PERSON]

	Live link client can run in cooked builds too

Change 3607668 by [REDACTED_PERSON]

Change 3607734 by [REDACTED_PERSON]

	LOD linking to curve
	- consolidated to one param - curve eval option
	- for long time, looking at why morphtarget wasn't working on LOD 1, later realized it was due to simplified :(
	- fixed to make sure param to clear is always checking with default value - this is correct behavior and it's not too bad for perf because internally the default value is also in the TMap
	- flipped meaning to align with bAllowCurveEvaluation
	- also fixed issue with orion cooking - where transform curves are added as normal curves

Change 3607859 by [REDACTED_PERSON]

	Missed files from live link editor checkin

Change 3607958 by [REDACTED_PERSON]

	Redo Jurre's changes from CL 3607229 (were removed by CL 3607734)

Change 3608566 by [REDACTED_PERSON]

	change include to avoid header conflicts on Linux

Change 3609074 by [REDACTED_PERSON]

	Take 2: Fix capitalization on include, fix Linux build.

Change 3610024 by [REDACTED_PERSON]

	Fix issue with material editor crashing due to missing load module of AdvancedPreviewScene

	- we used to load advanced preview setting by persona module
	- this has been moved to persona tool kit, and now all other modules are crashing
	- If we want to do it for tool kit, we have to make sure all other editor's loading should change also.

Change 3610081 by Jurre.deBaare

	Animations can't be set on blend samples from the dropdown
	#fix Skeleton asset registry tag now includes 'AssetTypeName' PathToAsset, so replacing compare with contains

Change 3610088 by Jurre.deBaare

	Editor crashes if you CtrlZ several times after adding animations to a 1D blendspace
	#fix removed the hacky OnObjectPropertyChanged and tied the refresh into propertyhandles instead
	#misc found out of sync widget values due to incorrect encapsulation inside of lambdas

Change 3610862 by [REDACTED_PERSON]

	Fix submix effects for situations where number of input channels does not equal output channels

Change 3611346 by [REDACTED_PERSON]

	Using audio thread platform affinity mask for audio render thread.

Change 3613297 by [REDACTED_PERSON]

	Simple delay submix

Change 3614435 by [REDACTED_PERSON]

	CIS fix

Change 3614482 by [REDACTED_PERSON]

	Store root motion on anim instance instead of proxy to avoid thread safety stalls

Change 3614483 by [REDACTED_PERSON]

	Evaluate curves in anim offsets

Change 3614495 by Jurre.deBaare

	Reimport alembic file with new source option does not automatically tick any tracks
	#fix If no tracks are set to import, reset them all to do so (we're assuming here the user is importing something completely different, and we wouldn't want her to import an empty animation either)

Change 3614645 by [REDACTED_PERSON]

	Fixed physics assets not simulating when BlockAll was globally overridden

	Persona viewport was overriding the collision profile back to BlockAll, which projects can override. Setting to the internal PhysicsActor profile prevents this, as it used to in PhAT

Change 3614683 by [REDACTED_PERSON]

	Fixed crash when modifying default physicsasset

Change 3614721 by Jurre.deBaare

	Vertex painting on skeletal meshes bound by physics asset
	#fix Now try and find intersecting triangle if we do hit the mesh bounds, but not any physics bodies

Change 3614730 by [REDACTED_PERSON]

	Fixed crash when regenerating multi convex hulls from zero-vert bones

	We handled this in the single convex hull case, but multi did not.

Change 3614763 by Jurre.deBaare

	Moving over:
	HLOD crash when dragging and dropping actors into their own cluster in the HLOD outliner - ALODActor
	#fix ensure that we nullptr check the static mesh as a LODActor can be reset to have a null static mesh

Change 3615029 by [REDACTED_PERSON]

	Fix issue with highlight

Change 3617593 by [REDACTED_PERSON]

	Fixed crash when regenerating large amounts of bodies

	We were overflowing the PhysX shape limit for aggregates - this refers to shapes, not bodies, it seems

Change 3617609 by [REDACTED_PERSON]

	Fixed crash that could occur when opening a physics asset and deleting bones.

Change 3617723 by [REDACTED_PERSON]

	Prevented actors & components of anim preview scenes (and the preview scenes themselves) from persisting after editors are shut down

	Fixed up 2 locations where the persona toolkit was being held onto by a strong ptr (cloth paint and new PhAT). This should stop the preview scene from persisting.
	Moved AddToRoot pattern used for anim preview scene to FGCObject

Change 3617818 by [REDACTED_PERSON]

	Final v1 UX changes for clothing tool, and removed experimental flag

Change 3617937 by Jurre.deBaare

	Default bounds for Alembic skel-mesh are too large
	#fix bounds was initialised to zero and +-ed which meant that it would always include (0,0,0) and enlarge the bounds

Change 3618187 by [REDACTED_PERSON]

	Implement Audiomixer in HTML5

Change 3618188 by [REDACTED_PERSON]

	Fix issue with highlight in persona

Change 3618229 by [REDACTED_PERSON]

	Fix crash on exit when modify is causing it to serialize again in the middle of tear down

Change 3618248 by [REDACTED_PERSON]

	fix issue by workaround where clamp is not happening with allowspin is false

Change 3618289 by [REDACTED_PERSON]

	Removing audio format types we're not using for simplicity

Change 3618291 by [REDACTED_PERSON]

	Fix duplicate of curve name appearing in list when renaming

Change 3618390 by [REDACTED_PERSON]

	Removing a case for DTYPE_Xenon since this is never used.

Change 3618425 by [REDACTED_PERSON]

	Keep notify UI up to data across multiple editors when adding notifies to an animation

Change 3619023 by [REDACTED_PERSON]

	Removing DTYPE_Xenon from XAudio2Buffer.cpp since it's not used

Change 3619129 by [REDACTED_PERSON]

	Source bus feature.

	- New architectural feature for audio mixer that allows audio sources to route to other audio sources.
	- Buses can be routed to each other
	- Buses have a duration which can be set in bus asset
	- Buses can choose between mono and stereo channels
	- Sources can send to buses and also toggle to *only* output to buses (and bypass submixing)
	- Will allow persistent source effects on different source audio, while also maintaining 3d spatialization capabilities. Lots of future features will build on this change: 3d audio-volume-based submixing, sidechaining, environment reflections, diagetic microphones, etc.

	- Some engine changes and optimizations:
	  - Format conversion to float is done in async workers for decode vs the render callback
	  - Procedural sound waves can opt to output only float vs int16 PCM data (avoids a format conversion in audio mixer)
	- Apply master attenuation at the final output vs per-source
	- Fixed code that performs fade in/fade out for smooth startup and shutdown.
	- Moved FSourceParam to FParam into DSP utility so others can use it.
	 - Some engine fixes:
	    - Audio spat plugins that are external sends will not send audio to default/base submix. But will also allow their audio to be panned and sent to submix sends (e.g. reverb) so external HRTF rendering can also get reverb effects, etc.
	  - Fixed an issue with pause
	  - Fixed an issue with the final source buffer in a source voice not getting properly rendered and causing discontinuties
	- Fixed an issue with WorldID not getting set for listeners

	TODO:
	 - fill out source bus details panel customization to hide USoundBase params which aren't relevant to source buses

Change 3619159 by [REDACTED_PERSON]

Change 3619555 by [REDACTED_PERSON]

	Fixed constraint debug drawing arrows in the physics asset editor being too large.

Change 3619574 by [REDACTED_PERSON]

	Fixed debug link for animation blueprints not persisting when changing preview mesh

	Anim instance is no longer re-created all the time when setting skeletal mesh, so we need to re-init the preview instance and re-set the linked skeletal mesh component manually when the mesh changes.

Change 3619586 by [REDACTED_PERSON]

	Fixed physics asset shortcut not working correctly in certain circumstances

	FBox was using uninitialized memory

Change 3619640 by [REDACTED_PERSON]

	Assets with no preview mesh now no longer allow access to other skeleton's physics assets in their shortcut bars

	Unified the skeleton/mesh search code between FPersonaAssetFamily and FPersonaToolkit, so they bot *look* for a compatible skeletal mesh if one was not found on the asset (but still dont set it automatically).

Change 3619644 by [REDACTED_PERSON]

	Change FBodyInstance::InstanceBodyIndex back to int32 (need to support ISMC with many instances)

Change 3619654 by [REDACTED_PERSON]

	Fix removing a curve when it isn't used on any animations

Change 3619771 by [REDACTED_PERSON]

	Make sure the physics asset editor floor has collision, regardless of what BlockAll does

Change 3619803 by [REDACTED_PERSON]

	Fixed localization warnings caused by duplicate keys.

Change 3619813 by Jurre.deBaare

	Baked bones using a pose animation are rotated in the wrong direction
	#fix root bone transform wasn't being taken into account while generating final bone transforms
	#misc added debug logging for future work

Change 3619830 by Jurre.deBaare

	Biased Texture Size option is not functioning when Merging Actors
	#fix Fixed up material baking setup after refactoring, now sets correct texture sizes again according to texture sizing type, this will be removed in the long term anyhow
	#misc Found a bug in material rendering  if previous render size < current render size it would not set the viewport size/projection matrix correctly which broke the material bake

Change 3619859 by [REDACTED_PERSON]

	Fixed HLOD selection sphere persisting on undo/redo

	Removed HLOD selection actors when the outliner is refreshed

Change 3619871 by [REDACTED_PERSON]

	Calculate root motion over the correct segment times, not the track times

Change 3619898 by [REDACTED_PERSON]

	Improve UI feedback around bounds/in-game bounds in animation editor viewports

	Tooltip for in-game bounds is now more detailed
	In-game bounds cannot be selected if bounds is not also selected

Change 3619908 by [REDACTED_PERSON]

	Fixed tooltip for PhysicsType

Change 3620014 by Jurre.deBaare

	Only the first mesh bake material property in the array can be set to diffuse, diffuse cannot be selected on the other array elements
	#fix Changed the way the restriction is setup and retrieve the UMaterialOptions from the details view instead of GetDefault<>
	#misc Also added more delegates to ensure the restriction is up to date

Change 3620104 by Jurre.deBaare

	HLOD doesn't support renaming in levels
	#fix ensure that during renaming of UWorld we also rename the HLOD assets into their respective new HLOD package outer

Change 3620151 by [REDACTED_PERSON]

	Undo/redo now correctly affects animation preview scene settings

	Preview scene desc is now transactional & state is correctly set up on undo/redo according to the current preview scene desc

Change 3620152 by [REDACTED_PERSON]

	Exposed LOD menu in PhAT

	This allows auto LOD to be optionaly selected. It was hidden and we forced to LOD 0 before. We still default to forcing LOD 0 to preserve the old behavior.

Change 3620177 by [REDACTED_PERSON]

	PR #3696: Fix for USkinnedMeshComponent::GetCPUSkinnedVertices() (Contributed by [REDACTED_PERSON])

Change 3620250 by Jurre.deBaare

	HLOD assets left in HLOD folder when clusters are deleted
	#fix some added lifetime management for HLOD assets, keeping list of 'stale' HLOD assets which if not Undo-ed will either be deleted when LODActor is saved, or marked PendingKill when LODActor is destroyed

Change 3620273 by [REDACTED_PERSON]

	PR #3908: Removing duplicated forward declation (Contributed by [REDACTED_PERSON])

Change 3620274 by [REDACTED_PERSON]

	PR #3909: Removing unnecessary conditional (Contributed by [REDACTED_PERSON])

Change 3620275 by [REDACTED_PERSON]

	Add icon for destruction plugin

Change 3620401 by [REDACTED_PERSON]

Change 3620586 by Jurre.deBaare

	Linux CIS fix

Change 3620660 by [REDACTED_PERSON]

	Fixes for state machines getting reinitialized in situations that users don't want them to.
	-Added option to state machine to allow it to skip reinitialization when it becomes relevant
	-Added option to slot nodes to keep source pose relevant while montage slot is playing.

Change 3620665 by [REDACTED_PERSON]

	Making source buses only show relevant source bus data.

	- hiding sound wave categories that aren't relevant to source buses

Change 3621087 by [REDACTED_PERSON]

Change 3621144 by [REDACTED_PERSON]

	 Fixing concurrency and volume stats

Change 3621148 by [REDACTED_PERSON]

	Fixing typo

Change 3621180 by [REDACTED_PERSON]

Change 3621421 by [REDACTED_PERSON]

Change 3621604 by [REDACTED_PERSON]

Change 3621626 by [REDACTED_PERSON]

	Fix issue with anim montage displaying when selecting animation

Change 3621813 by [REDACTED_PERSON]

	Fixing undo/redo of bone modifications in Physics Asset Editor (and others)

	Bone proxy objects now get recycled (instead of the pool constantly growing) as their names are stable and unique.
	Fixed broken skeleton tree RTTI (so selection persistance now works correctly on undo/redo again)
	We no longer force a re-selection on phyiscs asset changes (the tree takes care of that anyway).

Change 3621831 by Jurre.deBaare

	Crash fix for Material baking when trying to analyse a MP_MAX material property

Change 3621936 by [REDACTED_PERSON]

	Fixed CIS error from merge

Change 3621937 by [REDACTED_PERSON]

	Fix merge issue with API change in USynthComponent

Change 3622173 by [REDACTED_PERSON]

	Fixed ortho viewports being bright white in sub-editors

	Preview scenes in general are responsible by default for the background color.
	Advanced preview scenes now use background color from settings. Previously only te animation editors did this.

Change 3622730 by [REDACTED_PERSON]

Change 3622759 by [REDACTED_PERSON]

Change 3622901 by [REDACTED_PERSON]

	Bumped PhysX DDC key after change in Orion caused verify failures

Change 3623458 by [REDACTED_PERSON]

Change 3623667 by [REDACTED_PERSON]

Change 3623752 by [REDACTED_PERSON]

	Fixing issues with new source bus uobject so properties show up appropriately.

Change 3624058 by [REDACTED_PERSON]

	Fix stale module being enumerated when running UE4Editor-Cmd.exe, causing warning when running incremental automated tests.

	Module and version manifest filenames are derived from the executable filename, so when running the executable compiled for the console subsystem, we need to strip the -Cmd suffix from the executable name to find the correct path.

Change 3624193 by [REDACTED_PERSON]

Change 3354003 by [REDACTED_PERSON]

	Back out changelist 3353914

Change 3355932 by [REDACTED_PERSON]

	Back out changelist 3354003

	Reinstating merge from Main:

	Merging //UE4/Dev-Main to Dev-AnimPhys (//UE4/Dev-AnimPhys) @ CL 3353839

Change 3477632 by Jurre.deBaare

	Automated test content and ground truths for Actor Merging and Material baking functionality

Change 3491464 by Jurre.deBaare

	Updated automation content for MergeActor behaviour

Change 3587878 by [REDACTED_PERSON]

	Merging //UE4/Dev-Main to Dev-AnimPhys (//UE4/Dev-AnimPhys) @  CL 3587489

Change 3597452 by [REDACTED_PERSON]

Change 3602935 by [REDACTED_PERSON]

	Allow curve evaluation to be controlled by users

Change 3606120 by [REDACTED_PERSON]

	Move Tap Delay Submix to Synthesis library, modify tap delay API

Change 3621830 by [REDACTED_PERSON]

	Merging //UE4/Dev-Main to Dev-AnimPhys (//UE4/Dev-AnimPhys) @ CL 3621691

Change 3622807 by [REDACTED_PERSON]

	Issue is that these platforms weren't using the proper public function and an audio mixer refactor changed how volume is calculated to seperate out distance attenuation vs other volume gains.

diff --git a/Engine/Config/BaseEditor.ini b/Engine/Config/BaseEditor.ini
--- a/Engine/Config/BaseEditor.ini
+++ b/Engine/Config/BaseEditor.ini
@@ -39,6 +39,7 @@ MaxNumberOfLogLines=2000
 +RenderableThumbnailTypes=(ClassNeedingThumbnailName="/Script/Engine.SoundWave",RendererClassName="/Script/UnrealEd.SoundWaveThumbnailRenderer")
 +RenderableThumbnailTypes=(ClassNeedingThumbnailName="/Script/Engine.SlateBrushAsset",RendererClassName="/Script/UnrealEd.SlateBrushThumbnailRenderer")
 +RenderableThumbnailTypes=(ClassNeedingThumbnailName="/Script/Engine.SubsurfaceProfile",RendererClassName="/Script/UnrealEd.SubsurfaceProfileRenderer")
++RenderableThumbnailTypes=(ClassNeedingThumbnailName="/Script/Engine.PhysicsAsset",RendererClassName="/Script/UnrealEd.PhysicsAssetThumbnailRenderer")
 
 [/Script/UnrealEd.DefaultSizedThumbnailRenderer]
 DefaultSizeX=512
diff --git a/Engine/Config/BaseEditorPerProjectUserSettings.ini b/Engine/Config/BaseEditorPerProjectUserSettings.ini
--- a/Engine/Config/BaseEditorPerProjectUserSettings.ini
+++ b/Engine/Config/BaseEditorPerProjectUserSettings.ini
@@ -467,13 +467,9 @@ ZoomToScrubPos=false
 ShowCurveEd=false
 
 
-[/Script/UnrealEd.PhATSimOptions]
-SkyBrightness=0.25
-; for physical correct light computations we multiply diffuse and specular lights by PI (see LABEL_RealEnergy)
-Brightness=3.1415926535897932
+[/Script/UnrealEd.PhysicsAssetEditorOptions]
 AngularSnap=15.0
 LinearSnap=2.0
-SimSpeed=1.0
 bDrawContacts=false
 FloorGap=25.0
 GravScale=1.0
@@ -482,7 +478,8 @@ PokeStrength=100.0
 bShowNamesInHierarchy=true
 PokePauseTime=0.5
 PokeBlendTime=0.5
-ConstraintDrawSize=0.01
+ConstraintDrawSize=1.0
+bShowConstraintsAsPoints=false
 
 
 [/Script/UnrealEd.CurveEdOptions]
diff --git a/Engine/Documentation/Extras/API/API.ini b/Engine/Documentation/Extras/API/API.ini
--- a/Engine/Documentation/Extras/API/API.ini
+++ b/Engine/Documentation/Extras/API/API.ini
@@ -373,7 +373,7 @@ Description=Dialog used to display package status & operations (saving, checking
 [Module.Persona]
 Description=Editor for Skeletal meshes, skeletons and animation-related systems.
 
-[Module.PhAT]
+[Module.PhysicsAssetEditor]
 Description=Physical asset editor, allowing setup of physics for skeletal meshes.
 
 [Module.PlacementMode]
diff --git a/Engine/Documentation/Source/Shared/Editor/HLOD/HLOD.INT.udn b/Engine/Documentation/Source/Shared/Editor/HLOD/HLOD.INT.udn
--- a/Engine/Documentation/Source/Shared/Editor/HLOD/HLOD.INT.udn
+++ b/Engine/Documentation/Source/Shared/Editor/HLOD/HLOD.INT.udn
@@ -3,7 +3,10 @@ Title:HLOD System
 Crumbs:
 
 [EXCERPT:main]
-**Hierarchical Level of Detail** meshes (HLODs) have been developed as a way to increase performance and keep the quality bar high for projects using Unreal Engine 4.11 or later. HLODs can replace multiple **Static Mesh Actors** with single, combined **Static Mesh Actor** at long view distances. This helps reduce the number of **Actors** that need to be rendered for the scene, increasing performance by lowering the number of draw calls per frame. 
+###Hierarchical Level of Detail (HLODs)
+Has been developed as a way to increase performance and keep the quality bar high for projects using Unreal Engine 4.11 or later. 
+
+HLODs can replace multiple Static Mesh Actors with single, combined Static Mesh Actor at long view distances. This helps reduce the number of Actors that need to be rendered for the scene, increasing performance by lowering the number of draw calls per frame. 
 [VAR:ToolTipFullLink]
 Engine/HLOD
 [/VAR]
diff --git a/Engine/Plugins/Experimental/AlembicImporter/Source/AlembicLibrary/Private/AbcImportUtilities.cpp b/Engine/Plugins/Experimental/AlembicImporter/Source/AlembicLibrary/Private/AbcImportUtilities.cpp
--- a/Engine/Plugins/Experimental/AlembicImporter/Source/AlembicLibrary/Private/AbcImportUtilities.cpp
+++ b/Engine/Plugins/Experimental/AlembicImporter/Source/AlembicLibrary/Private/AbcImportUtilities.cpp
@@ -1296,7 +1296,7 @@ FBoxSphereBounds AbcImporterUtilities::ExtractBounds(Alembic::Abc::IBox3dPropert
 			const Imath::V3d BoundSize = BoundsSample.size();
 			const Imath::V3d BoundCenter = BoundsSample.center();
 			const FBoxSphereBounds ConvertedBounds(FVector(BoundCenter.x, BoundCenter.y, BoundCenter.z), FVector(BoundSize.x  * 0.5f, BoundSize.y * 0.5f, BoundSize.z * 0.5f), (const float)BoundSize.length() * 0.5f);
-			Bounds = Bounds + ConvertedBounds;			
+			Bounds = ( SampleIndex == 0 ) ? ConvertedBounds : Bounds + ConvertedBounds;
 		}
 	}
 
diff --git a/Engine/Plugins/Experimental/AlembicImporter/Source/AlembicLibrary/Private/AbcImporter.cpp b/Engine/Plugins/Experimental/AlembicImporter/Source/AlembicLibrary/Private/AbcImporter.cpp
--- a/Engine/Plugins/Experimental/AlembicImporter/Source/AlembicLibrary/Private/AbcImporter.cpp
+++ b/Engine/Plugins/Experimental/AlembicImporter/Source/AlembicLibrary/Private/AbcImporter.cpp
@@ -93,9 +93,24 @@ void FAbcImporter::UpdateAssetImportData(UAbcAssetImportData* AssetImportData)
 
 void FAbcImporter::RetrieveAssetImportData(UAbcAssetImportData* AssetImportData)
 {
+	bool bAnySetForImport = false;
+
 	for (TSharedPtr<FAbcPolyMeshObject>& MeshObject : ImportData->PolyMeshObjects)
 	{
-		MeshObject->bShouldImport = AssetImportData->TrackNames.Contains(MeshObject->Name);
+		if (AssetImportData->TrackNames.Contains(MeshObject->Name))
+		{
+			MeshObject->bShouldImport = true;
+			bAnySetForImport = true;
+		}		
+	}
+
+	// If none were set to import, set all of them to import (probably different scene/setup)
+	if (!bAnySetForImport)
+	{
+		for (TSharedPtr<FAbcPolyMeshObject>& MeshObject : ImportData->PolyMeshObjects)
+		{
+			MeshObject->bShouldImport = true;
+		}
 	}
 }
 
diff --git a/Engine/Plugins/Experimental/LiveLink/LiveLink.uplugin b/Engine/Plugins/Experimental/LiveLink/LiveLink.uplugin
--- a/Engine/Plugins/Experimental/LiveLink/LiveLink.uplugin
+++ b/Engine/Plugins/Experimental/LiveLink/LiveLink.uplugin
@@ -1,25 +1,30 @@
 {
-	"FileVersion" : 3,
-	"Version" : 1,
-	"VersionName" : "0.1",
-	"FriendlyName" : "Live Link",
-	"Description" : "Allows streaming of animation data into the editor",
-	"Category" : "Editor",
-	"CreatedBy" : "Epic Games, Inc.",
-	"CreatedByURL" : "http://epicgames.com",
-	"DocsURL" : "",
-	"MarketplaceURL" : "",
-	"SupportURL" : "",
-	"EnabledByDefault" : false,
-	"CanContainContent" : false,
-	"IsBetaVersion" : true,
-	"Installed" : false,
-	"Modules" :
-	[
-		{
-			"Name" : "LiveLink",
-			"Type" : "Editor",
-			"LoadingPhase" : "PreDefault"
-		}
-	]
+  "FileVersion" : 3,
+  "Version" : 1,
+  "VersionName" : "0.1",
+  "FriendlyName" : "Live Link",
+  "Description" : "Allows streaming of animation data into the editor",
+  "Category" : "Animation",
+  "CreatedBy" : "Epic Games, Inc.",
+  "CreatedByURL" : "http://epicgames.com",
+  "DocsURL" : "",
+  "MarketplaceURL" : "",
+  "SupportURL" : "",
+  "EnabledByDefault" : false,
+  "CanContainContent" : false,
+  "IsBetaVersion" : true,
+  "Installed" : false,
+  "Modules" :
+  [
+    {
+      "Name": "LiveLink",
+      "Type": "Runtime",
+      "LoadingPhase": "PreDefault"
+    },
+    {
+      "Name" : "LiveLinkEditor",
+      "Type" : "Editor",
+      "LoadingPhase" : "PreDefault"
+    }
+  ]
 }
\ No newline at end of file
diff --git a/Engine/Plugins/Experimental/LiveLink/Source/LiveLink/LiveLink.Build.cs b/Engine/Plugins/Experimental/LiveLink/Source/LiveLink/LiveLink.Build.cs
--- a/Engine/Plugins/Experimental/LiveLink/Source/LiveLink/LiveLink.Build.cs
+++ b/Engine/Plugins/Experimental/LiveLink/Source/LiveLink/LiveLink.Build.cs
@@ -2,34 +2,25 @@
 
 namespace UnrealBuildTool.Rules
 {
-	public class LiveLink : ModuleRules
-	{
-		public LiveLink(ReadOnlyTargetRules Target) : base(Target)
+    public class LiveLink : ModuleRules
+    {
+        public LiveLink(ReadOnlyTargetRules Target) : base(Target)
         {
-			PrivateDependencyModuleNames.AddRange(
-				new string[]
-				{
-					"Core",
-					"CoreUObject",
-                    "UnrealEd",
-                    "Engine",
-                    "Projects",
+            PrivateDependencyModuleNames.AddRange(
+            new string[]
+            {
+                "Core",
+                "CoreUObject",
+                "Engine",
+                "Projects",
 
-                    "WorkspaceMenuStructure",
-                    "EditorStyle",
-                    "SlateCore",
-                    "Slate",
-                    "InputCore",
+                "Messaging",
+                "LiveLinkInterface",
+                "LiveLinkMessageBusFramework",
+            }
+            );
 
-                    "Messaging",
-                    "LiveLinkInterface",
-					"LiveLinkMessageBusFramework",
-                    "BlueprintGraph",
-                    "AnimGraph",
-                }
-			); 
-
-			PrivateIncludePaths.Add("/LiveLink/Private");
-		}
-	}
+            PrivateIncludePaths.Add("/LiveLink/Private");
+        }
+    }
 }
diff --git a/Engine/Plugins/Experimental/LiveLink/Source/LiveLink/Private/AnimNode_LiveLinkPose.cpp b/Engine/Plugins/Experimental/LiveLink/Source/LiveLink/Private/AnimNode_LiveLinkPose.cpp
--- a/Engine/Plugins/Experimental/LiveLink/Source/LiveLink/Private/AnimNode_LiveLinkPose.cpp
+++ b/Engine/Plugins/Experimental/LiveLink/Source/LiveLink/Private/AnimNode_LiveLinkPose.cpp
@@ -9,8 +9,6 @@
 
 #include "LiveLinkRemapAsset.h"
 
-#define LOCTEXT_NAMESPACE "LiveLinkAnimNode"
-
 FAnimNode_LiveLinkPose::FAnimNode_LiveLinkPose() 
 	: RetargetAsset(ULiveLinkRemapAsset::StaticClass())
 	, LiveLinkClient(nullptr)
@@ -78,25 +76,3 @@ void FAnimNode_LiveLinkPose::OnLiveLinkClientUnregistered(const FName& Type, cla
 		LiveLinkClient = nullptr;
 	}
 }
-
-UAnimGraphNode_LiveLinkPose::UAnimGraphNode_LiveLinkPose(const FObjectInitializer& ObjectInitializer)
-	: Super(ObjectInitializer)
-{
-}
-
-FText UAnimGraphNode_LiveLinkPose::GetNodeTitle(ENodeTitleType::Type TitleType) const
-{
-	return LOCTEXT("NodeTitle", "Live Link Pose");
-}
-
-FText UAnimGraphNode_LiveLinkPose::GetTooltipText() const
-{
-	return LOCTEXT("NodeTooltip", "Retrieves the current pose associated with the supplied subject");
-}
-
-FText UAnimGraphNode_LiveLinkPose::GetMenuCategory() const
-{
-	return LOCTEXT("NodeCategory", "Live Link");
-}
-
-#undef LOCTEXT_NAMESPACE
diff --git a/Engine/Plugins/Experimental/LiveLink/Source/LiveLink/Private/LiveLinkModule.cpp b/Engine/Plugins/Experimental/LiveLink/Source/LiveLink/Private/LiveLinkModule.cpp
--- a/Engine/Plugins/Experimental/LiveLink/Source/LiveLink/Private/LiveLinkModule.cpp
+++ b/Engine/Plugins/Experimental/LiveLink/Source/LiveLink/Private/LiveLinkModule.cpp
@@ -2,27 +2,9 @@
 
 #include "ILiveLinkModule.h"
 
-#include "IPluginManager.h"
-
-#include "Editor.h"
-
-#include "ModuleManager.h"
-//#include "CoreDelegates.h"
-#include "Widgets/Docking/SDockTab.h"
-#include "WorkspaceMenuStructure.h"
-#include "WorkspaceMenuStructureModule.h"
-#include "EditorStyleSet.h"
-#include "SlateStyle.h"
-#include "SlateTypes.h"
-#include "SlateStyleRegistry.h"
-#include "Framework/Application/SlateApplication.h"
-
 #include "Features/IModularFeatures.h"
-#include "LevelEditor.h"
 
 #include "LiveLinkClient.h"
-#include "LiveLinkClientPanel.h"
-#include "LiveLinkClientCommands.h"
 
 /**
  * Implements the Messaging module.
@@ -30,109 +12,28 @@
 
 #define LOCTEXT_NAMESPACE "LiveLinkModule"
 
-static const FName LiveLinkClientTabName(TEXT("LiveLink"));
-
-#define IMAGE_PLUGIN_BRUSH( RelativePath, ... ) FSlateImageBrush( InPluginContent( RelativePath, ".png" ), __VA_ARGS__ )
-
-FString InPluginContent(const FString& RelativePath, const ANSICHAR* Extension)
-{
-	static FString ContentDir = IPluginManager::Get().FindPlugin(TEXT("LiveLink"))->GetContentDir();
-	return (ContentDir / RelativePath) + Extension;
-}
-
 class FLiveLinkModule
 	: public ILiveLinkModule
 {
 public:
 	FLiveLinkClient LiveLinkClient;
 
-	TSharedPtr<FSlateStyleSet> StyleSet;
-
-	TSharedPtr< class ISlateStyle > GetStyleSet() { return StyleSet; }
-
 	// IModuleInterface interface
 
 	virtual void StartupModule() override
 	{
 		IModularFeatures::Get().RegisterModularFeature(FLiveLinkClient::ModularFeatureName, &LiveLinkClient);
-
-		FLevelEditorModule& LevelEditorModule = FModuleManager::LoadModuleChecked<FLevelEditorModule>(TEXT("LevelEditor"));
-
-		FLiveLinkClient* ClientPtr = &LiveLinkClient;
-
-		static FName LiveLinkStyle(TEXT("LiveLinkStyle"));
-		StyleSet = MakeShareable(new FSlateStyleSet(LiveLinkStyle));
-
-		FLiveLinkClientCommands::Register();
-
-		TSharedPtr<FSlateStyleSet> StyleSetPtr = StyleSet;
-
-		//Register our UI
-		LevelEditorTabManagerChangedHandle = LevelEditorModule.OnTabManagerChanged().AddLambda([ClientPtr, StyleSetPtr]()
-		{
-			FLevelEditorModule& LocalLevelEditorModule = FModuleManager::GetModuleChecked<FLevelEditorModule>(TEXT("LevelEditor"));
-			LocalLevelEditorModule.GetLevelEditorTabManager()->RegisterTabSpawner(LiveLinkClientTabName, FOnSpawnTab::CreateStatic(&FLiveLinkModule::SpawnLiveLinkTab, ClientPtr, StyleSetPtr))
-				.SetGroup(WorkspaceMenu::GetMenuStructure().GetLevelEditorCategory())
-				.SetDisplayName(LOCTEXT("LiveLinkTabTitle", "Live Link"))
-				.SetTooltipText(LOCTEXT("SequenceRecorderTooltipText", "Open the Live Link streaming manager tab."))
-				.SetIcon(FSlateIcon(StyleSetPtr->GetStyleSetName(), "LiveLinkClient.Common.Icon.Small"));
-		});
-
-		const FVector2D Icon16x16(16.0f, 16.0f);
-		const FVector2D Icon20x20(20.0f, 20.0f);
-		const FVector2D Icon40x40(40.0f, 40.0f);
-
-		StyleSet->SetContentRoot(FPaths::EngineContentDir() / TEXT("Editor/Slate"));
-		StyleSet->SetCoreContentRoot(FPaths::EngineContentDir() / TEXT("Slate"));
-
-		StyleSet->Set("LiveLinkClient.Common.Icon", new IMAGE_PLUGIN_BRUSH(TEXT("LiveLink_40x"), Icon40x40));
-		StyleSet->Set("LiveLinkClient.Common.Icon.Small", new IMAGE_PLUGIN_BRUSH(TEXT("LiveLink_16x"), Icon16x16));
-
-		StyleSet->Set("LiveLinkClient.Common.AddSource", new IMAGE_PLUGIN_BRUSH(TEXT("icon_AddSource_40x"), Icon40x40));
-		StyleSet->Set("LiveLinkClient.Common.RemoveSource", new IMAGE_PLUGIN_BRUSH(TEXT("icon_RemoveSource_40x"), Icon40x40));
-		StyleSet->Set("LiveLinkClient.Common.RemoveAllSources", new IMAGE_PLUGIN_BRUSH(TEXT("icon_RemoveSource_40x"), Icon40x40));
-
-		FSlateStyleRegistry::RegisterSlateStyle(*StyleSet.Get());
 	}
 
 	virtual void ShutdownModule() override
 	{
 		IModularFeatures::Get().UnregisterModularFeature(FLiveLinkClient::ModularFeatureName, &LiveLinkClient);
-
-		if (FSlateApplication::IsInitialized())
-		{
-			FGlobalTabmanager::Get()->UnregisterTabSpawner(LiveLinkClientTabName);
-		}
-
-		if (FModuleManager::Get().IsModuleLoaded(TEXT("LevelEditor")))
-		{
-			FLevelEditorModule& LevelEditorModule = FModuleManager::GetModuleChecked<FLevelEditorModule>(TEXT("LevelEditor"));
-			LevelEditorModule.OnTabManagerChanged().Remove(LevelEditorTabManagerChangedHandle);
-		}
 	}
 
 	virtual bool SupportsDynamicReloading() override
 	{
 		return false;
 	}
-
-	static TSharedRef<SDockTab> SpawnLiveLinkTab(const FSpawnTabArgs& SpawnTabArgs, FLiveLinkClient* Client, TSharedPtr<FSlateStyleSet> StyleSet)
-	{
-		const FSlateBrush* IconBrush = StyleSet->GetBrush("LiveLinkClient.Common.Icon.Small");
-
-		const TSharedRef<SDockTab> MajorTab =
-			SNew(SDockTab)
-			.Icon(IconBrush)
-			.TabRole(ETabRole::NomadTab);
-
-		MajorTab->SetContent(SNew(SLiveLinkClientPanel, Client));
-
-		return MajorTab;
-	}
-
-private:
-
-	FDelegateHandle LevelEditorTabManagerChangedHandle;
 };
 
 IMPLEMENT_MODULE(FLiveLinkModule, LiveLink);
diff --git a/Engine/Plugins/Experimental/LiveLink/Source/LiveLink/Public/AnimNode_LiveLinkPose.h b/Engine/Plugins/Experimental/LiveLink/Source/LiveLink/Public/AnimNode_LiveLinkPose.h
--- a/Engine/Plugins/Experimental/LiveLink/Source/LiveLink/Public/AnimNode_LiveLinkPose.h
+++ b/Engine/Plugins/Experimental/LiveLink/Source/LiveLink/Public/AnimNode_LiveLinkPose.h
@@ -4,7 +4,6 @@
 
 #include "CoreMinimal.h"
 #include "Animation/AnimNodeBase.h"
-#include "AnimGraphNode_Base.h"
 #include "LiveLinkRetargetAsset.h"
 
 class ILiveLinkClient;
@@ -43,19 +42,4 @@ struct LIVELINK_API FAnimNode_LiveLinkPose : public FAnimNode_Base
 private:
 
 	ILiveLinkClient* LiveLinkClient;
-};
-
-UCLASS()
-class UAnimGraphNode_LiveLinkPose : public UAnimGraphNode_Base
-{
-	GENERATED_UCLASS_BODY()
-
-	UPROPERTY(EditAnywhere, Category = Settings)
-	FAnimNode_LiveLinkPose Node;
-
-	// UEdGraphNode interface
-	virtual FText GetNodeTitle(ENodeTitleType::Type TitleType) const override;
-	virtual FText GetTooltipText() const override;
-	virtual FText GetMenuCategory() const;
-	// End of UEdGraphNode
 };
\ No newline at end of file
diff --git a/Engine/Plugins/Experimental/LiveLink/Source/LiveLink/Public/ILiveLinkModule.h b/Engine/Plugins/Experimental/LiveLink/Source/LiveLink/Public/ILiveLinkModule.h
--- a/Engine/Plugins/Experimental/LiveLink/Source/LiveLink/Public/ILiveLinkModule.h
+++ b/Engine/Plugins/Experimental/LiveLink/Source/LiveLink/Public/ILiveLinkModule.h
@@ -29,8 +29,6 @@ public:
 #endif
 	}
 
-	virtual TSharedPtr< class ISlateStyle > GetStyleSet() = 0;
-
 public:
 
 	/** Virtual destructor. */
diff --git a/Engine/Plugins/Experimental/LiveLink/Source/LiveLink/Public/LiveLinkClient.h b/Engine/Plugins/Experimental/LiveLink/Source/LiveLink/Public/LiveLinkClient.h
--- a/Engine/Plugins/Experimental/LiveLink/Source/LiveLink/Public/LiveLinkClient.h
+++ b/Engine/Plugins/Experimental/LiveLink/Source/LiveLink/Public/LiveLinkClient.h
@@ -4,7 +4,7 @@
 
 #include "Transform.h"
 #include "CriticalSection.h"
-#include "TickableEditorObject.h"
+#include "Tickable.h"
 
 #include "ILiveLinkClient.h"
 #include "ILiveLinkSource.h"
@@ -70,7 +70,7 @@ struct FLiveLinkSubject
 	void BuildInterpolatedFrame(const double InSeconds, FLiveLinkSubjectFrame& OutFrame);
 };
 
-class FLiveLinkClient : public ILiveLinkClient, public FTickableEditorObject
+class LIVELINK_API FLiveLinkClient : public ILiveLinkClient, public FTickableGameObject
 {
 public:
 	FLiveLinkClient() : LastValidationCheck(0.0) {}
@@ -78,13 +78,12 @@ public:
 
 	// Begin FTickableObjectBase implementation
 	virtual bool IsTickable() const override { return true; }
+	virtual bool IsTickableWhenPaused() const override { return true; }
+	virtual bool IsTickableInEditor() const override { return true; }
 	virtual void Tick(float DeltaTime) override;
 	virtual TStatId GetStatId() const override { RETURN_QUICK_DECLARE_CYCLE_STAT(LiveLinkClient, STATGROUP_Tickables); }
 	// End FTickableObjectBase
 
-	// Add a new live link source to the client
-	void AddSource(TSharedPtr<ILiveLinkSource> InSource);
-
 	// Remove the specified source from the live link client
 	void RemoveSource(FGuid InEntryGuid);
 
@@ -92,8 +91,10 @@ public:
 	void RemoveAllSources();
 
 	// ILiveLinkClient Interface
-	virtual FLiveLinkTimeCode MakeTimeCode(double InTime, int32 InFrameNum) const;
-	virtual FLiveLinkTimeCode MakeTimeCodeFromTimeOnly(double InTime) const;
+	virtual FLiveLinkTimeCode MakeTimeCode(double InTime, int32 InFrameNum) const override;
+	virtual FLiveLinkTimeCode MakeTimeCodeFromTimeOnly(double InTime) const override;
+
+	virtual void AddSource(TSharedPtr<ILiveLinkSource> InSource) override;
 
 	virtual void PushSubjectSkeleton(FName SubjectName, const FLiveLinkRefSkeleton& RefSkeleton) override;
 	virtual void ClearSubject(FName SubjectName) override;
diff --git a/Engine/Plugins/Experimental/LiveLink/Source/LiveLink/Public/LiveLinkMessageBusSource.h b/Engine/Plugins/Experimental/LiveLink/Source/LiveLink/Public/LiveLinkMessageBusSource.h
--- a/Engine/Plugins/Experimental/LiveLink/Source/LiveLink/Public/LiveLinkMessageBusSource.h
+++ b/Engine/Plugins/Experimental/LiveLink/Source/LiveLink/Public/LiveLinkMessageBusSource.h
@@ -13,7 +13,7 @@ struct FLiveLinkSubjectFrameMessage;
 struct FLiveLinkHeartbeatMessage;
 struct FLiveLinkClearSubject;
 
-class FLiveLinkMessageBusSource : public ILiveLinkSource
+class LIVELINK_API FLiveLinkMessageBusSource : public ILiveLinkSource
 {
 public:
 
diff --git a/Engine/Plugins/Experimental/LiveLink/Source/LiveLinkEditor/Private/AnimGraphNode_LiveLinkPose.h b/Engine/Plugins/Experimental/LiveLink/Source/LiveLinkEditor/Private/AnimGraphNode_LiveLinkPose.h
new file mode 100644
--- /dev/null
+++ b/Engine/Plugins/Experimental/LiveLink/Source/LiveLinkEditor/Private/AnimGraphNode_LiveLinkPose.h
@@ -0,0 +1,22 @@
+// Copyright 1998-2017 Epic Games, Inc. All Rights Reserved.
+
+#pragma once
+
+#include "AnimGraphNode_Base.h"
+#include "AnimNode_LiveLinkPose.h"
+
+#include "AnimGraphNode_LiveLinkPose.generated.h"
+UCLASS()
+class UAnimGraphNode_LiveLinkPose : public UAnimGraphNode_Base
+{
+	GENERATED_UCLASS_BODY()
+
+	UPROPERTY(EditAnywhere, Category = Settings)
+	FAnimNode_LiveLinkPose Node;
+
+	// UEdGraphNode interface
+	virtual FText GetNodeTitle(ENodeTitleType::Type TitleType) const override;
+	virtual FText GetTooltipText() const override;
+	virtual FText GetMenuCategory() const;
+	// End of UEdGraphNode
+};
\ No newline at end of file
diff --git a/Engine/Plugins/Experimental/LiveLink/Source/LiveLinkEditor/Private/LiveLinkClientCommands.cpp b/Engine/Plugins/Experimental/LiveLink/Source/LiveLinkEditor/Private/LiveLinkClientCommands.cpp
new file mode 100644
--- /dev/null
+++ b/Engine/Plugins/Experimental/LiveLink/Source/LiveLinkEditor/Private/LiveLinkClientCommands.cpp
@@ -0,0 +1,20 @@
+// Copyright 1998-2017 Epic Games, Inc. All Rights Reserved.
+
+#include "LiveLinkClientCommands.h"
+#include "EditorStyleSet.h"
+
+#define LOCTEXT_NAMESPACE "LiveLinkClientCommands"
+
+FLiveLinkClientCommands::FLiveLinkClientCommands()
+	: TCommands<FLiveLinkClientCommands>("LiveLinkClient.Common", LOCTEXT("Common", "Common"), NAME_None, FName(TEXT("LiveLinkStyle")))
+{
+}
+
+void FLiveLinkClientCommands::RegisterCommands()
+{
+	UI_COMMAND(AddSource, "Add", "Add a new live link source", EUserInterfaceActionType::Button, FInputChord());
+	UI_COMMAND(RemoveSource, "Remove", "Remove selected live link source", EUserInterfaceActionType::Button, FInputChord());
+	UI_COMMAND(RemoveAllSources, "Remove All", "Remove all live link sources", EUserInterfaceActionType::Button, FInputChord());
+}
+
+#undef LOCTEXT_NAMESPACE
diff --git a/Engine/Plugins/Experimental/LiveLink/Source/LiveLinkEditor/Private/LiveLinkClientPanel.cpp b/Engine/Plugins/Experimental/LiveLink/Source/LiveLinkEditor/Private/LiveLinkClientPanel.cpp
new file mode 100644
--- /dev/null
+++ b/Engine/Plugins/Experimental/LiveLink/Source/LiveLinkEditor/Private/LiveLinkClientPanel.cpp
@@ -0,0 +1,403 @@
+// Copyright 1998-2017 Epic Games, Inc. All Rights Reserved.
+
+#include "LiveLinkClientPanel.h"
+#include "SBoxPanel.h"
+#include "SSplitter.h"
+#include "SOverlay.h"
+#include "SharedPointer.h"
+#include "SListView.h"
+#include "SButton.h"
+#include "SlateApplication.h"
+#include "MultiBoxBuilder.h"
+#include "SlateDelegates.h"
+#include "LiveLinkClientCommands.h"
+#include "ILiveLinkSource.h"
+#include "LiveLinkClient.h"
+#include "UObjectHash.h"
+#include "EditorStyleSet.h"
+#include "ModuleManager.h"
+#include "PropertyEditorModule.h"
+#include "IStructureDetailsView.h"
+
+#include "LiveLinkSourceFactory.h"
+
+#define LOCTEXT_NAMESPACE "LiveLinkClientPanel"
+
+static const FName TypeColumnName(TEXT("Type"));
+static const FName MachineColumnName(TEXT("Machine"));
+static const FName StatusColumnName(TEXT("Status"));
+
+struct FLiveLinkSourceUIEntry
+{
+public:
+	FLiveLinkSourceUIEntry(FGuid InEntryGuid, FLiveLinkClient* InClient)
+		: EntryGuid(InEntryGuid)
+		, Client(InClient)
+	{}
+
+	FText GetSourceType() { return Client->GetSourceTypeForEntry(EntryGuid); }
+	FText GetMachineName() { return Client->GetMachineNameForEntry(EntryGuid); }
+	FText GetEntryStatus() { return Client->GetEntryStatusForEntry(EntryGuid); }
+	FLiveLinkConnectionSettings* GetConnectionSettings() { return Client->GetConnectionSettingsForEntry(EntryGuid); }
+	void RemoveFromClient() { Client->RemoveSource(EntryGuid); }
+
+private:
+	FGuid EntryGuid;
+	FLiveLinkClient* Client;
+};
+
+class SLiveLinkClientPanelSourcesRow : public SMultiColumnTableRow<FLiveLinkSourceUIEntryPtr>
+{
+public:
+	SLATE_BEGIN_ARGS(SLiveLinkClientPanelSourcesRow) {}
+
+		/** The list item for this row */
+		SLATE_ARGUMENT(FLiveLinkSourceUIEntryPtr, Entry)
+
+	SLATE_END_ARGS()
+
+	void Construct(const FArguments& Args, const TSharedRef<STableViewBase>& OwnerTableView)
+	{
+		EntryPtr = Args._Entry;
+
+		SourceType = EntryPtr->GetSourceType();
+		MachineName = EntryPtr->GetMachineName();
+
+		SMultiColumnTableRow<FLiveLinkSourceUIEntryPtr>::Construct(
+			FSuperRowType::FArguments()
+			.Padding(1.0f),
+			OwnerTableView
+		);
+	}
+
+	/** Overridden from SMultiColumnTableRow.  Generates a widget for this column of the list view. */
+	virtual TSharedRef<SWidget> GenerateWidgetForColumn(const FName& ColumnName) override
+	{
+		if (ColumnName == TypeColumnName)
+		{
+			return	SNew(STextBlock)
+					.Text(TAttribute<FText>::Create(TAttribute<FText>::FGetter::CreateSP(this, &SLiveLinkClientPanelSourcesRow::GetSourceTypeName)));
+		}
+		else if (ColumnName == MachineColumnName)
+		{
+			return	SNew(STextBlock)
+					.Text(TAttribute<FText>::Create(TAttribute<FText>::FGetter::CreateSP(this, &SLiveLinkClientPanelSourcesRow::GetSourceMachineName)));
+		}
+		else if (ColumnName == StatusColumnName)
+		{
+			return	SNew(STextBlock)
+					.Text(TAttribute<FText>::Create(TAttribute<FText>::FGetter::CreateSP(this, &SLiveLinkClientPanelSourcesRow::GetSourceStatus)));
+		}
+
+		return SNullWidget::NullWidget;
+	}
+
+private:
+
+	FText GetSourceTypeName() const
+	{
+		return SourceType;
+	}
+
+	FText GetSourceMachineName() const
+	{
+		return MachineName;
+	}
+
+	FText GetSourceStatus() const
+	{
+		return EntryPtr->GetEntryStatus();
+	}
+
+	FLiveLinkSourceUIEntryPtr EntryPtr;
+
+	FText SourceType;
+
+	FText MachineName;
+};
+
+SLiveLinkClientPanel::~SLiveLinkClientPanel()
+{
+	if (Client)
+	{
+		Client->UnregisterSourcesChangedHandle(OnSourcesChangedHandle);
+		OnSourcesChangedHandle.Reset();
+	}
+}
+
+void SLiveLinkClientPanel::Construct(const FArguments& Args, FLiveLinkClient* InClient)
+{
+	check(InClient);
+	Client = InClient;
+
+	OnSourcesChangedHandle = Client->RegisterSourcesChangedHandle(FLiveLinkSourcesChanged::FDelegate::CreateSP(this, &SLiveLinkClientPanel::OnSourcesChangedHandler));
+
+	RefreshSourceData(false);
+
+	CommandList = MakeShareable(new FUICommandList);
+
+	BindCommands();
+
+	FToolBarBuilder ToolBarBuilder(CommandList, FMultiBoxCustomization::None);
+
+	ToolBarBuilder.BeginSection(TEXT("Add"));
+	{
+		ToolBarBuilder.AddComboButton(	FUIAction(),
+										FOnGetContent::CreateSP(this, &SLiveLinkClientPanel::GenerateSourceMenu),
+										LOCTEXT("AddSource", "Add"),
+										LOCTEXT("AddSource_ToolTip", "Add a new live link source"),
+										FSlateIcon(FName("LiveLinkStyle"), "LiveLinkClient.Common.AddSource")
+										);
+	}
+	ToolBarBuilder.EndSection();
+
+	ToolBarBuilder.BeginSection(TEXT("Remove"));
+	{
+		ToolBarBuilder.AddToolBarButton(FLiveLinkClientCommands::Get().RemoveSource);
+		ToolBarBuilder.AddToolBarButton(FLiveLinkClientCommands::Get().RemoveAllSources);
+	}
+	ToolBarBuilder.EndSection();
+
+	// Connection Settings
+	FPropertyEditorModule& PropertyEditorModule = FModuleManager::GetModuleChecked<FPropertyEditorModule>("PropertyEditor");
+
+	FDetailsViewArgs DetailsViewArgs;
+	//DetailsViewArgs.
+	FStructureDetailsViewArgs StructureViewArgs;
+	StructureViewArgs.bShowAssets = true;
+	StructureViewArgs.bShowClasses = true;
+	StructureViewArgs.bShowInterfaces = true;
+	StructureViewArgs.bShowObjects = true;
+
+	StructureDetailsView = PropertyEditorModule.CreateStructureDetailView(DetailsViewArgs, StructureViewArgs, nullptr);
+
+	ChildSlot
+	[
+		SNew(SSplitter)
+		.Orientation(Orient_Vertical)
+		+ SSplitter::Slot()
+		.Value(0.33f)
+		[
+			SNew(SVerticalBox)
+			+ SVerticalBox::Slot()
+			.AutoHeight()
+			.Padding(FMargin(0.0f, 4.0f, 0.0f, 0.0f))
+			[
+				ToolBarBuilder.MakeWidget()
+			]
+			+ SVerticalBox::Slot()
+			.FillHeight(0.5f)
+			.Padding(FMargin(0.0f, 4.0f, 0.0f, 0.0f))
+			[
+				SNew(SBorder)
+				.BorderImage(FEditorStyle::GetBrush("ToolPanel.GroupBorder"))
+				.Padding(FMargin(4.0f, 4.0f))
+				[
+					SNew(SOverlay)
+					+ SOverlay::Slot()
+					[
+						SNew(SVerticalBox)
+						+ SVerticalBox::Slot()
+						.FillHeight(1.0f)
+						[
+							SAssignNew(ListView, SListView<FLiveLinkSourceUIEntryPtr>)
+							.ListItemsSource(&SourceData)
+							.SelectionMode(ESelectionMode::SingleToggle)
+							.OnGenerateRow(this, &SLiveLinkClientPanel::MakeSourceListViewWidget)
+							.OnSelectionChanged(this, &SLiveLinkClientPanel::OnSourceListSelectionChanged)
+							.HeaderRow
+							(
+								SNew(SHeaderRow)
+								+ SHeaderRow::Column(TypeColumnName)
+								.FillWidth(43.0f)
+								.DefaultLabel(LOCTEXT("TypeColumnHeaderName", "Source Type"))
+								+ SHeaderRow::Column(MachineColumnName)
+								.FillWidth(43.0f)
+								.DefaultLabel(LOCTEXT("MachineColumnHeaderName", "Source Machine"))
+								+ SHeaderRow::Column(StatusColumnName)
+								.FillWidth(14.0f)
+								.DefaultLabel(LOCTEXT("StatusColumnHeaderName", "Status"))
+							)
+						]
+					]
+				]
+			]
+			+ SVerticalBox::Slot()
+			.FillHeight(0.5f)
+			.Padding(FMargin(0.0f, 4.0f, 0.0f, 0.0f))
+			[
+				StructureDetailsView->GetWidget().ToSharedRef()
+			]
+		]
+	];
+
+	// Register refresh timer
+	/*if (!ActiveTimerHandle.IsValid())
+	{
+		ActiveTimerHandle = RegisterActiveTimer(0.1f, FWidgetActiveTimerDelegate::CreateSP(this, &SSequenceRecorder::HandleRefreshItems));
+	}*/
+}
+
+void SLiveLinkClientPanel::BindCommands()
+{
+	CommandList->MapAction(FLiveLinkClientCommands::Get().RemoveSource,
+		FExecuteAction::CreateSP(this, &SLiveLinkClientPanel::HandleRemoveSource),
+		FCanExecuteAction::CreateSP(this, &SLiveLinkClientPanel::CanRemoveSource)
+	);
+
+	CommandList->MapAction(FLiveLinkClientCommands::Get().RemoveAllSources,
+		FExecuteAction::CreateSP(this, &SLiveLinkClientPanel::HandleRemoveAllSources),
+		FCanExecuteAction::CreateSP(this, &SLiveLinkClientPanel::CanRemoveSource)
+	);
+}
+
+void SLiveLinkClientPanel::RefreshSourceData(bool bRefreshUI)
+{
+	SourceData.Reset();
+
+	for (FGuid Source : Client->GetSourceEntries())
+	{
+		SourceData.Add(MakeShareable(new FLiveLinkSourceUIEntry(Source, Client)));
+	}
+
+	if (bRefreshUI)
+	{
+		ListView->RequestListRefresh();
+	}
+}
+
+const TArray<FLiveLinkSourceUIEntryPtr>& SLiveLinkClientPanel::GetCurrentSources() const
+{
+	return SourceData;
+}
+
+TSharedRef<ITableRow> SLiveLinkClientPanel::MakeSourceListViewWidget(FLiveLinkSourceUIEntryPtr Entry, const TSharedRef<STableViewBase>& OwnerTable) const
+{
+	return SNew(SLiveLinkClientPanelSourcesRow, OwnerTable)
+		   .Entry(Entry);
+}
+
+void SLiveLinkClientPanel::OnSourceListSelectionChanged(FLiveLinkSourceUIEntryPtr Entry, ESelectInfo::Type SelectionType) const
+{
+	if(Entry.IsValid())
+	{
+		FStructOnScope* Struct = new FStructOnScope(FLiveLinkConnectionSettings::StaticStruct(), (uint8*)Entry->GetConnectionSettings());
+		StructureDetailsView->SetStructureData(MakeShareable(Struct));
+	}
+	else
+	{
+		StructureDetailsView->SetStructureData(nullptr);
+	}
+}
+
+TSharedRef< SWidget > SLiveLinkClientPanel::GenerateSourceMenu()
+{
+	TArray<UClass*> Results;
+	GetDerivedClasses(ULiveLinkSourceFactory::StaticClass(), Results, true);
+
+	const bool CloseAfterSelection = true;
+	FMenuBuilder MenuBuilder(CloseAfterSelection, NULL);
+
+	MenuBuilder.BeginSection("SourceSection", LOCTEXT("Sources", "Live Link Sources"));
+
+	for (UClass* SourceFactory : Results)
+	{
+		ULiveLinkSourceFactory* FactoryCDO = SourceFactory->GetDefaultObject<ULiveLinkSourceFactory>();
+
+		//Prep UI
+		TSharedPtr<SWidget> SourcePanel = FactoryCDO->CreateSourceCreationPanel();
+		SourcePanels.Add(FactoryCDO) = SourcePanel;
+
+		MenuBuilder.AddSubMenu(
+			FactoryCDO->GetSourceDisplayName(),
+			FactoryCDO->GetSourceTooltip(),
+			FNewMenuDelegate::CreateRaw(this, &SLiveLinkClientPanel::RetrieveFactorySourcePanel, FactoryCDO),
+			false
+		);
+	}
+
+	MenuBuilder.EndSection();
+	return MenuBuilder.MakeWidget();
+}
+
+void SLiveLinkClientPanel::RetrieveFactorySourcePanel(FMenuBuilder& MenuBuilder, ULiveLinkSourceFactory* FactoryCDO)
+{
+	TSharedPtr<SWidget> SourcePanel = SourcePanels[FactoryCDO];
+	MenuBuilder.AddWidget(
+		SNew(SVerticalBox)
+		+ SVerticalBox::Slot()
+		.FillHeight(1)
+		[
+			SourcePanel.ToSharedRef()
+		]
+		+SVerticalBox::Slot()
+		.AutoHeight()
+		[
+			SNew(SHorizontalBox)
+			+SHorizontalBox::Slot()
+			.FillWidth(1.f)
+			+SHorizontalBox::Slot()
+			.AutoWidth()
+			.HAlign(EHorizontalAlignment::HAlign_Right)
+			[
+				SNew(SButton)
+				.Text(LOCTEXT("OkButton", "Ok"))
+				.OnClicked(this, &SLiveLinkClientPanel::OnCloseSourceSelectionPanel, FactoryCDO, true)
+			]
+			+ SHorizontalBox::Slot()
+			.AutoWidth()
+			.HAlign(EHorizontalAlignment::HAlign_Right)
+			[
+				SNew(SButton)
+				.Text(LOCTEXT("CancelButton", "Cancel"))
+				.OnClicked(this, &SLiveLinkClientPanel::OnCloseSourceSelectionPanel, FactoryCDO, false)
+			]
+		]
+		, FText(), true);
+}
+
+FReply SLiveLinkClientPanel::OnCloseSourceSelectionPanel(ULiveLinkSourceFactory* FactoryCDO, bool bMakeSource)
+{
+	TSharedPtr<ILiveLinkSource> Source = FactoryCDO->OnSourceCreationPanelClosed(bMakeSource);
+	// If we want a source we should get one ... if we dont we shouldn't. Make sure source factory does the right thing in both cases
+	if (bMakeSource)
+	{
+		check(Source.IsValid());
+		Client->AddSource(Source);
+
+		RefreshSourceData(true);
+	}
+	else
+	{
+		check(!Source.IsValid());
+	}
+	FSlateApplication::Get().DismissAllMenus();
+	return FReply::Handled();
+}
+
+void SLiveLinkClientPanel::HandleRemoveSource()
+{
+	TArray<FLiveLinkSourceUIEntryPtr> Selected;
+	ListView->GetSelectedItems(Selected);
+	if (Selected.Num() > 0)
+	{
+		Selected[0]->RemoveFromClient();
+	}
+}
+
+bool SLiveLinkClientPanel::CanRemoveSource()
+{
+	return ListView->GetNumItemsSelected() > 0;
+}
+
+void SLiveLinkClientPanel::HandleRemoveAllSources()
+{
+	Client->RemoveAllSources();
+}
+
+void SLiveLinkClientPanel::OnSourcesChangedHandler()
+{
+	RefreshSourceData(true);
+}
+
+#undef LOCTEXT_NAMESPACE
\ No newline at end of file
diff --git a/Engine/Plugins/Experimental/LiveLink/Source/LiveLinkEditor/Private/LiveLinkClientPanel.h b/Engine/Plugins/Experimental/LiveLink/Source/LiveLinkEditor/Private/LiveLinkClientPanel.h
new file mode 100644
--- /dev/null
+++ b/Engine/Plugins/Experimental/LiveLink/Source/LiveLinkEditor/Private/LiveLinkClientPanel.h
@@ -0,0 +1,76 @@
+// Copyright 1998-2017 Epic Games, Inc. All Rights Reserved.
+
+#pragma once
+
+#include "Widgets/SCompoundWidget.h"
+#include "Widgets/DeclarativeSyntaxSupport.h"
+#include "SListView.h"
+
+struct FLiveLinkSourceUIEntry;
+class FLiveLinkClient;
+class ULiveLinkSourceFactory;
+class ITableRow;
+class STableViewBase;
+class FMenuBuilder;
+class FUICommandList;
+
+typedef TSharedPtr<FLiveLinkSourceUIEntry> FLiveLinkSourceUIEntryPtr;
+
+class SLiveLinkClientPanel : public SCompoundWidget
+{
+	SLATE_BEGIN_ARGS(SLiveLinkClientPanel)
+	{
+	}
+
+	SLATE_END_ARGS()
+
+	~SLiveLinkClientPanel();
+
+	void Construct(const FArguments& Args, FLiveLinkClient* InClient);
+
+private:
+
+	void BindCommands();
+
+	void RefreshSourceData(bool bRefreshUI);
+
+	TSharedRef< SWidget > GenerateSourceMenu();
+
+	void RetrieveFactorySourcePanel(FMenuBuilder& MenuBuilder, ULiveLinkSourceFactory* FactoryCDO);
+
+	FReply OnCloseSourceSelectionPanel(ULiveLinkSourceFactory* FactoryCDO, bool bMakeSource);
+
+	void HandleOnAddFromFactory(ULiveLinkSourceFactory* InSourceFactory);
+
+	void HandleRemoveSource();
+
+	bool CanRemoveSource();
+
+	void HandleRemoveAllSources();
+
+	// Registered with the client and called when clients sources changes
+	void OnSourcesChangedHandler();
+
+private:
+	const TArray<FLiveLinkSourceUIEntryPtr>& GetCurrentSources() const;
+
+	TSharedRef<ITableRow> MakeSourceListViewWidget(FLiveLinkSourceUIEntryPtr Entry, const TSharedRef<STableViewBase>& OwnerTable) const;
+	
+	void OnSourceListSelectionChanged(FLiveLinkSourceUIEntryPtr Entry, ESelectInfo::Type SelectionType) const;
+	
+	TSharedPtr<SListView<FLiveLinkSourceUIEntryPtr>> ListView;
+
+	TArray<FLiveLinkSourceUIEntryPtr> SourceData;
+
+	TSharedPtr<FUICommandList> CommandList;
+
+	FLiveLinkClient* Client;
+
+	TMap<ULiveLinkSourceFactory*, TSharedPtr<SWidget>> SourcePanels;
+
+	// Reference to connection settings struct details panel
+	TSharedPtr<class IStructureDetailsView> StructureDetailsView;
+
+	// Handle to delegate registered with client so we can update when a source disappears
+	FDelegateHandle OnSourcesChangedHandle;
+};
\ No newline at end of file
diff --git a/Engine/Plugins/Experimental/LiveLink/Source/LiveLinkEditor/Private/LiveLinkEditorModule.cpp b/Engine/Plugins/Experimental/LiveLink/Source/LiveLinkEditor/Private/LiveLinkEditorModule.cpp
new file mode 100644
--- /dev/null
+++ b/Engine/Plugins/Experimental/LiveLink/Source/LiveLinkEditor/Private/LiveLinkEditorModule.cpp
@@ -0,0 +1,131 @@
+// Copyright 1998-2017 Epic Games, Inc. All Rights Reserved.
+
+#include "IPluginManager.h"
+
+#include "Editor.h"
+
+#include "ModuleManager.h"
+#include "CoreDelegates.h"
+#include "Widgets/Docking/SDockTab.h"
+#include "WorkspaceMenuStructure.h"
+#include "WorkspaceMenuStructureModule.h"
+#include "EditorStyleSet.h"
+#include "SlateStyle.h"
+#include "SlateTypes.h"
+#include "SlateStyleRegistry.h"
+#include "Framework/Application/SlateApplication.h"
+
+#include "Features/IModularFeatures.h"
+#include "LevelEditor.h"
+
+#include "LiveLinkClient.h"
+#include "LiveLinkClientPanel.h"
+#include "LiveLinkClientCommands.h"
+
+/**
+ * Implements the Messaging module.
+ */
+
+#define LOCTEXT_NAMESPACE "LiveLinkModule"
+
+static const FName LiveLinkClientTabName(TEXT("LiveLink"));
+
+#define IMAGE_PLUGIN_BRUSH( RelativePath, ... ) FSlateImageBrush( InPluginContent( RelativePath, ".png" ), __VA_ARGS__ )
+
+FString InPluginContent(const FString& RelativePath, const ANSICHAR* Extension)
+{
+	static FString ContentDir = IPluginManager::Get().FindPlugin(TEXT("LiveLink"))->GetContentDir();
+	return (ContentDir / RelativePath) + Extension;
+}
+
+class FLiveLinkEditorModule : public IModuleInterface
+{
+public:
+	TSharedPtr<FSlateStyleSet> StyleSet;
+
+	TSharedPtr< class ISlateStyle > GetStyleSet() { return StyleSet; }
+
+	// IModuleInterface interface
+
+	virtual void StartupModule() override
+	{
+		FLevelEditorModule& LevelEditorModule = FModuleManager::LoadModuleChecked<FLevelEditorModule>(TEXT("LevelEditor"));
+
+		static FName LiveLinkStyle(TEXT("LiveLinkStyle"));
+		StyleSet = MakeShareable(new FSlateStyleSet(LiveLinkStyle));
+
+		FLiveLinkClientCommands::Register();
+
+		TSharedPtr<FSlateStyleSet> StyleSetPtr = StyleSet;
+
+		//Register our UI
+		LevelEditorTabManagerChangedHandle = LevelEditorModule.OnTabManagerChanged().AddLambda([StyleSetPtr]()
+		{
+			FLevelEditorModule& LocalLevelEditorModule = FModuleManager::GetModuleChecked<FLevelEditorModule>(TEXT("LevelEditor"));
+			LocalLevelEditorModule.GetLevelEditorTabManager()->RegisterTabSpawner(LiveLinkClientTabName, FOnSpawnTab::CreateStatic(&FLiveLinkEditorModule::SpawnLiveLinkTab, StyleSetPtr))
+				.SetGroup(WorkspaceMenu::GetMenuStructure().GetLevelEditorCategory())
+				.SetDisplayName(LOCTEXT("LiveLinkTabTitle", "Live Link"))
+				.SetTooltipText(LOCTEXT("SequenceRecorderTooltipText", "Open the Live Link streaming manager tab."))
+				.SetIcon(FSlateIcon(StyleSetPtr->GetStyleSetName(), "LiveLinkClient.Common.Icon.Small"));
+		});
+
+		const FVector2D Icon16x16(16.0f, 16.0f);
+		const FVector2D Icon20x20(20.0f, 20.0f);
+		const FVector2D Icon40x40(40.0f, 40.0f);
+
+		StyleSet->SetContentRoot(FPaths::EngineContentDir() / TEXT("Editor/Slate"));
+		StyleSet->SetCoreContentRoot(FPaths::EngineContentDir() / TEXT("Slate"));
+
+		StyleSet->Set("LiveLinkClient.Common.Icon", new IMAGE_PLUGIN_BRUSH(TEXT("LiveLink_40x"), Icon40x40));
+		StyleSet->Set("LiveLinkClient.Common.Icon.Small", new IMAGE_PLUGIN_BRUSH(TEXT("LiveLink_16x"), Icon16x16));
+
+		StyleSet->Set("LiveLinkClient.Common.AddSource", new IMAGE_PLUGIN_BRUSH(TEXT("icon_AddSource_40x"), Icon40x40));
+		StyleSet->Set("LiveLinkClient.Common.RemoveSource", new IMAGE_PLUGIN_BRUSH(TEXT("icon_RemoveSource_40x"), Icon40x40));
+		StyleSet->Set("LiveLinkClient.Common.RemoveAllSources", new IMAGE_PLUGIN_BRUSH(TEXT("icon_RemoveSource_40x"), Icon40x40));
+
+		FSlateStyleRegistry::RegisterSlateStyle(*StyleSet.Get());
+	}
+
+	virtual void ShutdownModule() override
+	{
+		if (FSlateApplication::IsInitialized())
+		{
+			FGlobalTabmanager::Get()->UnregisterTabSpawner(LiveLinkClientTabName);
+		}
+
+		if (FModuleManager::Get().IsModuleLoaded(TEXT("LevelEditor")))
+		{
+			FLevelEditorModule& LevelEditorModule = FModuleManager::GetModuleChecked<FLevelEditorModule>(TEXT("LevelEditor"));
+			LevelEditorModule.OnTabManagerChanged().Remove(LevelEditorTabManagerChangedHandle);
+		}
+	}
+
+	virtual bool SupportsDynamicReloading() override
+	{
+		return false;
+	}
+
+	static TSharedRef<SDockTab> SpawnLiveLinkTab(const FSpawnTabArgs& SpawnTabArgs, TSharedPtr<FSlateStyleSet> StyleSet)
+	{
+		FLiveLinkClient* Client = &IModularFeatures::Get().GetModularFeature<FLiveLinkClient>(FLiveLinkClient::ModularFeatureName);
+
+		const FSlateBrush* IconBrush = StyleSet->GetBrush("LiveLinkClient.Common.Icon.Small");
+
+		const TSharedRef<SDockTab> MajorTab =
+			SNew(SDockTab)
+			.Icon(IconBrush)
+			.TabRole(ETabRole::NomadTab);
+
+		MajorTab->SetContent(SNew(SLiveLinkClientPanel, Client));
+
+		return MajorTab;
+	}
+
+private:
+
+	FDelegateHandle LevelEditorTabManagerChangedHandle;
+};
+
+IMPLEMENT_MODULE(FLiveLinkEditorModule, LiveLinkEditor);
+
+#undef LOCTEXT_NAMESPACE
\ No newline at end of file
diff --git a/Engine/Plugins/Experimental/LiveLink/Source/LiveLinkEditor/Private/LiveLinkMessageBusSourceEditor.cpp b/Engine/Plugins/Experimental/LiveLink/Source/LiveLinkEditor/Private/LiveLinkMessageBusSourceEditor.cpp
new file mode 100644
--- /dev/null
+++ b/Engine/Plugins/Experimental/LiveLink/Source/LiveLinkEditor/Private/LiveLinkMessageBusSourceEditor.cpp
@@ -0,0 +1,121 @@
+// Copyright 1998-2017 Epic Games, Inc. All Rights Reserved.
+
+#include "LiveLinkMessageBusSourceEditor.h"
+#include "LiveLinkMessages.h"
+#include "SBox.h"
+
+#include "MessageEndpointBuilder.h"
+
+#define LOCTEXT_NAMESPACE "LiveLinkMessageBusSourceEditor"
+
+namespace ProviderPollUI
+{
+	const FName TypeColumnName(TEXT("Type"));
+	const FName MachineColumnName(TEXT("Machine"));
+};
+
+class SProviderPollRow : public SMultiColumnTableRow<FProviderPollResultPtr>
+{
+public:
+	SLATE_BEGIN_ARGS(SProviderPollRow) {}
+
+		/** The list item for this row */
+		SLATE_ARGUMENT(FProviderPollResultPtr, PollResultPtr)
+
+	SLATE_END_ARGS()
+
+	void Construct(const FArguments& Args, const TSharedRef<STableViewBase>& OwnerTableView)
+	{
+		PollResultPtr = Args._PollResultPtr;
+
+		SMultiColumnTableRow<FProviderPollResultPtr>::Construct(
+			FSuperRowType::FArguments()
+			.Padding(1.0f),
+			OwnerTableView
+		);
+	}
+
+	/** Overridden from SMultiColumnTableRow.  Generates a widget for this column of the list view. */
+	virtual TSharedRef<SWidget> GenerateWidgetForColumn(const FName& ColumnName) override
+	{
+		if (ColumnName == ProviderPollUI::TypeColumnName)
+		{
+			return	SNew(STextBlock)
+					.Text(FText::FromString(PollResultPtr->Name));
+		}
+		else if (ColumnName == ProviderPollUI::MachineColumnName)
+		{
+			return	SNew(STextBlock)
+				.Text(FText::FromString(PollResultPtr->MachineName));
+		}
+
+		return SNullWidget::NullWidget;
+	}
+
+private:
+
+	FProviderPollResultPtr PollResultPtr;
+};
+
+
+SLiveLinkMessageBusSourceEditor::~SLiveLinkMessageBusSourceEditor()
+{
+	FMessageEndpoint::SafeRelease(MessageEndpoint);
+}
+
+void SLiveLinkMessageBusSourceEditor::Construct(const FArguments& Args)
+{
+	MessageEndpoint =	FMessageEndpoint::Builder(TEXT("LiveLinkMessageBusSource"))
+						.Handling<FLiveLinkPongMessage>(this, &SLiveLinkMessageBusSourceEditor::HandlePongMessage);
+
+	CurrentPollRequest = FGuid::NewGuid();
+
+	//Simple each for connections till UI comes along
+	MessageEndpoint->Publish(new FLiveLinkPingMessage(CurrentPollRequest));
+
+	ChildSlot
+	[
+		SNew(SBox)
+		.HeightOverride(200)
+		.WidthOverride(200)
+		[
+			SAssignNew(ListView, SListView<FProviderPollResultPtr>)
+			.ListItemsSource(&PollData)
+			.SelectionMode(ESelectionMode::SingleToggle)
+			.OnGenerateRow(this, &SLiveLinkMessageBusSourceEditor::MakeSourceListViewWidget)
+			.OnSelectionChanged(this, &SLiveLinkMessageBusSourceEditor::OnSourceListSelectionChanged)
+			.HeaderRow
+			(
+				SNew(SHeaderRow)
+				+ SHeaderRow::Column(ProviderPollUI::TypeColumnName)
+				.FillWidth(43.0f)
+				.DefaultLabel(LOCTEXT("TypeColumnHeaderName", "Source Type"))
+				+ SHeaderRow::Column(ProviderPollUI::MachineColumnName)
+				.FillWidth(43.0f)
+				.DefaultLabel(LOCTEXT("MachineColumnHeaderName", "Source Machine"))
+			)
+		]
+	];
+}
+
+TSharedRef<ITableRow> SLiveLinkMessageBusSourceEditor::MakeSourceListViewWidget(FProviderPollResultPtr PollResult, const TSharedRef<STableViewBase>& OwnerTable) const
+{
+	return SNew(SProviderPollRow, OwnerTable).PollResultPtr(PollResult);
+}
+
+void SLiveLinkMessageBusSourceEditor::OnSourceListSelectionChanged(FProviderPollResultPtr PollResult, ESelectInfo::Type SelectionType)
+{
+	SelectedResult = PollResult;
+}
+
+void SLiveLinkMessageBusSourceEditor::HandlePongMessage(const FLiveLinkPongMessage& Message, const TSharedRef<IMessageContext, ESPMode::ThreadSafe>& Context)
+{
+	if(Message.PollRequest == CurrentPollRequest)
+	{
+		PollData.Add(MakeShareable(new FProviderPollResult(Context->GetSender(), Message.ProviderName, Message.MachineName)));
+
+		ListView->RequestListRefresh();
+	}
+}
+
+#undef LOCTEXT_NAMESPACE
\ No newline at end of file
diff --git a/Engine/Plugins/Experimental/LiveLink/Source/LiveLinkEditor/Private/LiveLinkMessageBusSourceEditor.h b/Engine/Plugins/Experimental/LiveLink/Source/LiveLinkEditor/Private/LiveLinkMessageBusSourceEditor.h
new file mode 100644
--- /dev/null
+++ b/Engine/Plugins/Experimental/LiveLink/Source/LiveLinkEditor/Private/LiveLinkMessageBusSourceEditor.h
@@ -0,0 +1,65 @@
+// Copyright 1998-2017 Epic Games, Inc. All Rights Reserved.
+
+#pragma once
+
+#include "Widgets/SCompoundWidget.h"
+#include "SListView.h"
+#include "IMessageContext.h"
+#include "MessageEndpoint.h"
+#include "Widgets/DeclarativeSyntaxSupport.h"
+#include "Guid.h"
+
+struct FLiveLinkPongMessage;
+struct FMessageAddress;
+class ITableRow;
+class STableViewBase;
+
+struct FProviderPollResult
+{
+public:
+	FProviderPollResult(const FMessageAddress& InAddress, const FString& InName, const FString& InMachineName)
+		: Address(InAddress)
+		, Name(InName)
+		, MachineName(InMachineName)
+	{}
+
+	FMessageAddress Address;
+
+	FString			Name;
+	FString			MachineName;
+};
+
+typedef TSharedPtr<FProviderPollResult> FProviderPollResultPtr;
+
+class SLiveLinkMessageBusSourceEditor : public SCompoundWidget
+{
+	SLATE_BEGIN_ARGS(SLiveLinkMessageBusSourceEditor)
+	{
+	}
+
+	SLATE_END_ARGS()
+
+	~SLiveLinkMessageBusSourceEditor();
+
+	void Construct(const FArguments& Args);
+
+	FProviderPollResultPtr GetSelectedSource() const { return SelectedResult; }
+
+private:
+
+	void HandlePongMessage(const FLiveLinkPongMessage& Message, const TSharedRef<IMessageContext, ESPMode::ThreadSafe>& Context);
+
+	TSharedRef<ITableRow> MakeSourceListViewWidget(FProviderPollResultPtr PollResult, const TSharedRef<STableViewBase>& OwnerTable) const;
+
+	void OnSourceListSelectionChanged(FProviderPollResultPtr PollResult, ESelectInfo::Type SelectionType);
+
+	TSharedPtr<SListView<FProviderPollResultPtr>> ListView;
+
+	TArray<FProviderPollResultPtr> PollData;
+
+	FProviderPollResultPtr SelectedResult;
+
+	TSharedPtr<FMessageEndpoint, ESPMode::ThreadSafe> MessageEndpoint;
+
+	FGuid CurrentPollRequest;
+};
\ No newline at end of file
diff --git a/Engine/Plugins/Experimental/LiveLink/Source/LiveLinkEditor/Private/LiveLinkMessageBusSourceFactory.cpp b/Engine/Plugins/Experimental/LiveLink/Source/LiveLinkEditor/Private/LiveLinkMessageBusSourceFactory.cpp
new file mode 100644
--- /dev/null
+++ b/Engine/Plugins/Experimental/LiveLink/Source/LiveLinkEditor/Private/LiveLinkMessageBusSourceFactory.cpp
@@ -0,0 +1,44 @@
+// Copyright 1998-2017 Epic Games, Inc. All Rights Reserved.
+#include "LiveLinkMessageBusSourceFactory.h"
+#include "LiveLinkMessageBusSource.h"
+#include "LiveLinkMessageBusSourceEditor.h"
+
+#define LOCTEXT_NAMESPACE "LiveLinkMessageBusSourceFactory"
+
+FText ULiveLinkMessageBusSourceFactory::GetSourceDisplayName() const
+{
+	return LOCTEXT("SourceDisplayName", "Message Bus Source");
+}
+
+FText ULiveLinkMessageBusSourceFactory::GetSourceTooltip() const
+{
+	return LOCTEXT("SourceTooltip", "Creates a connection to a Message Bus based Live Link Source");
+}
+
+TSharedPtr<SWidget> ULiveLinkMessageBusSourceFactory::CreateSourceCreationPanel()
+{
+	if (!ActiveSourceEditor.IsValid())
+	{
+		SAssignNew(ActiveSourceEditor, SLiveLinkMessageBusSourceEditor);
+	}
+	return ActiveSourceEditor;
+}
+
+TSharedPtr<ILiveLinkSource> ULiveLinkMessageBusSourceFactory::OnSourceCreationPanelClosed(bool bMakeSource)
+{
+	//Clean up
+	TSharedPtr<FLiveLinkMessageBusSource> NewSource = nullptr;
+
+	if (bMakeSource && ActiveSourceEditor.IsValid())
+	{
+		FProviderPollResultPtr Result = ActiveSourceEditor->GetSelectedSource();
+		if(Result.IsValid())
+		{
+			NewSource = MakeShareable( new FLiveLinkMessageBusSource(FText::FromString(Result->Name), FText::FromString(Result->MachineName), Result->Address));
+		}
+	}
+	ActiveSourceEditor = nullptr;
+	return NewSource;
+}
+
+#undef LOCTEXT_NAMESPACE
\ No newline at end of file
diff --git a/Engine/Plugins/Experimental/LiveLink/Source/LiveLinkEditor/Private/LiveLinkMessageBusSourceFactory.h b/Engine/Plugins/Experimental/LiveLink/Source/LiveLinkEditor/Private/LiveLinkMessageBusSourceFactory.h
new file mode 100644
--- /dev/null
+++ b/Engine/Plugins/Experimental/LiveLink/Source/LiveLinkEditor/Private/LiveLinkMessageBusSourceFactory.h
@@ -0,0 +1,25 @@
+// Copyright 1998-2017 Epic Games, Inc. All Rights Reserved.
+
+#pragma once
+
+#include "LiveLinkSourceFactory.h"
+#include "LiveLinkMessageBusSourceFactory.generated.h"
+
+
+class SLiveLinkMessageBusSourceEditor;
+
+UCLASS()
+class ULiveLinkMessageBusSourceFactory : public ULiveLinkSourceFactory
+{
+public:
+
+	GENERATED_BODY()
+
+	virtual FText GetSourceDisplayName() const;
+	virtual FText GetSourceTooltip() const;
+
+	virtual TSharedPtr<SWidget> CreateSourceCreationPanel();
+	virtual TSharedPtr<ILiveLinkSource> OnSourceCreationPanelClosed(bool bMakeSource);
+
+	TSharedPtr<SLiveLinkMessageBusSourceEditor> ActiveSourceEditor;
+};
\ No newline at end of file
diff --git a/Engine/Plugins/Runtime/Oculus/OculusAudio/Source/OculusAudio/Private/IOculusAudioPlugin.cpp b/Engine/Plugins/Runtime/Oculus/OculusAudio/Source/OculusAudio/Private/IOculusAudioPlugin.cpp
new file mode 100644
--- /dev/null
+++ b/Engine/Plugins/Runtime/Oculus/OculusAudio/Source/OculusAudio/Private/IOculusAudioPlugin.cpp
@@ -0,0 +1,16 @@
+// Copyright 1998-2017 Epic Games, Inc. All Rights Reserved.
+
+#include "IOculusAudioPlugin.h"
+
+void FOculusAudioPlugin::StartupModule()
+{
+	FOculusAudioLibraryManager::Initialize();
+	IModularFeatures::Get().RegisterModularFeature(FOculusSpatializationPluginFactory::GetModularFeatureName(), &PluginFactory);
+};
+
+void FOculusAudioPlugin::ShutdownModule()
+{
+	FOculusAudioLibraryManager::Shutdown();
+}
+
+IMPLEMENT_MODULE(FOculusAudioPlugin, OculusAudio)
\ No newline at end of file
diff --git a/Engine/Plugins/Runtime/Oculus/OculusAudio/Source/OculusAudio/Private/OculusAudio.cpp b/Engine/Plugins/Runtime/Oculus/OculusAudio/Source/OculusAudio/Private/OculusAudio.cpp
--- a/Engine/Plugins/Runtime/Oculus/OculusAudio/Source/OculusAudio/Private/OculusAudio.cpp
+++ b/Engine/Plugins/Runtime/Oculus/OculusAudio/Source/OculusAudio/Private/OculusAudio.cpp
@@ -1,443 +1,13 @@
-// Copyright 1998-2017 Epic Games, Inc. All Rights Reserved.
-
 #include "OculusAudio.h"
-#include "OVR_Audio.h"
-#include "XAudio2Device.h"
-#include "AudioEffect.h"
-#include "XAudio2Effects.h"
-#include "OculusAudioEffect.h"
-#include "Runtime/Windows/XAudio2/Private/XAudio2Support.h"
-#include "Misc/Paths.h"
-
-//---------------------------------------------------
-// Oculus Audio Plugin Implementation
-//---------------------------------------------------
-
-static const TCHAR* GetOculusErrorString(ovrResult Result)
-{
-	switch (Result)
-	{
-		default:
-		case ovrError_AudioUnknown:			return TEXT("Unknown Error");
-		case ovrError_AudioInvalidParam:	return TEXT("Invalid Param");
-		case ovrError_AudioBadSampleRate:	return TEXT("Bad Samplerate");
-		case ovrError_AudioMissingDLL:		return TEXT("Missing DLL");
-		case ovrError_AudioBadAlignment:	return TEXT("Pointers did not meet 16 byte alignment requirements");
-		case ovrError_AudioUninitialized:	return TEXT("Function called before initialization");
-		case ovrError_AudioHRTFInitFailure: return TEXT("HRTF Profider initialization failed");
-		case ovrError_AudioBadVersion:		return TEXT("Bad audio version");
-		case ovrError_AudioSRBegin:			return TEXT("Sample rate begin");
-		case ovrError_AudioSREnd:			return TEXT("Sample rate end");
-	}
-}
-
-#define OVR_AUDIO_CHECK(Result, Context)																\
-	if (Result != ovrSuccess)																			\
-	{																									\
-		const TCHAR* ErrString = GetOculusErrorString(Result);											\
-		UE_LOG(LogAudio, Error, TEXT("Oculus Audio SDK Error - %s: %s"), TEXT(Context), ErrString);	\
-		return;																							\
-	}
-
-class FHrtfSpatializationAlgorithmXAudio2 : public IAudioSpatialization
-{
-public:
-	FHrtfSpatializationAlgorithmXAudio2(class FAudioDevice * AudioDevice)
-		: bOvrContextInitialized(false)
-		, AudioDevice(AudioDevice)
-		, OvrAudioContext(nullptr)
-	{
-		check(HRTFEffects.Num() == 0);
-		check(Params.Num() == 0);
-
-		for (int32 EffectIndex = 0; EffectIndex < (int32)AudioDevice->MaxChannels; ++EffectIndex)
-		{
-			FXAudio2HRTFEffect* NewHRTFEffect = new FXAudio2HRTFEffect(EffectIndex, AudioDevice);
-			NewHRTFEffect->Initialize(nullptr, 0);
-			HRTFEffects.Add(NewHRTFEffect);
-
-			Params.Add(FSpatializationParams());
-		}
-	}
-
-	~FHrtfSpatializationAlgorithmXAudio2()
-	{
-		// Release all the effects for the oculus spatialization effect
-		for (FXAudio2HRTFEffect* Effect : HRTFEffects)
-		{
-			if (Effect)
-			{
-				delete Effect;
-			}
-		}
-		HRTFEffects.Empty();
-
-		// Destroy the contexts if we created them
-		if (bOvrContextInitialized)
-		{
-			if (OvrAudioContext != nullptr)
-			{
-				ovrAudio_DestroyContext(OvrAudioContext);
-				OvrAudioContext = nullptr;
-			}
-
-			bOvrContextInitialized = false;
-		}
-	}
-
-	bool IsSpatializationEffectInitialized() const override
-	{
-		return bOvrContextInitialized;
-	}
-
-	void InitializeSpatializationEffect(uint32 BufferLength) override
-	{
-		if (bOvrContextInitialized)
-		{
-			return;
-		}
-
-		uint32 MaxNumSources = AudioDevice->MaxChannels;
-		uint32 SampleRate = AudioDevice->SampleRate;
-
-		ovrAudioContextConfiguration ContextConfig;
-		ContextConfig.acc_Size = sizeof(ovrAudioContextConfiguration);
-
-		// First initialize the Fast algorithm context
-		ContextConfig.acc_Provider = ovrAudioSpatializationProvider_OVR_OculusHQ;
-		ContextConfig.acc_MaxNumSources = MaxNumSources;
-		ContextConfig.acc_SampleRate = SampleRate;
-		ContextConfig.acc_BufferLength = BufferLength;
-
-		check(OvrAudioContext == nullptr);
-		// Create the OVR Audio Context with a given quality
-		ovrResult Result = ovrAudio_CreateContext(&OvrAudioContext, &ContextConfig);
-		OVR_AUDIO_CHECK(Result, "Failed to create simple context");
-
-		// Now initialize the high quality algorithm context
-		bOvrContextInitialized = true;
-	}
-
-	void ProcessSpatializationForVoice(uint32 VoiceIndex, float* InSamples, float* OutSamples, const FVector& Position) override
-	{
-		if (OvrAudioContext)
-		{
-			ProcessAudioInternal(OvrAudioContext, VoiceIndex, InSamples, OutSamples, Position);
-		}
-	}
-
-	virtual bool CreateSpatializationEffect(uint32 VoiceId) override
-	{
-		// If an effect for this voice has already been created, then leave
-		if ((int32)VoiceId >= HRTFEffects.Num())
-		{
-			return false;
-		}
-		return true;
-	}
-
-	void* GetSpatializationEffect(uint32 VoiceId) override
-	{
-		if ((int32)VoiceId < HRTFEffects.Num())
-		{
-			return (void*)HRTFEffects[VoiceId];
-		}
-		return nullptr;
-	}
-
-	void SetSpatializationParameters(uint32 VoiceId, const FSpatializationParams& InParams) override
-	{
-		check((int32)VoiceId < Params.Num());
-
-		FScopeLock ScopeLock(&ParamCriticalSection);
-		Params[VoiceId] = InParams;
-	}
-
-	void GetSpatializationParameters(uint32 VoiceId, FSpatializationParams& OutParams) override
-	{
-		check((int32)VoiceId < Params.Num());
-
-		FScopeLock ScopeLock(&ParamCriticalSection);
-		OutParams = Params[VoiceId];
-	}
-
-private:
-
-	void ProcessAudioInternal(ovrAudioContext AudioContext, uint32 VoiceIndex, float* InSamples, float* OutSamples, const FVector& Position)
-	{
-		ovrResult Result = ovrAudio_SetAudioSourceAttenuationMode(AudioContext, VoiceIndex, ovrAudioSourceAttenuationMode_None, 1.0f);
-		OVR_AUDIO_CHECK(Result, "Failed to set source attenuation mode");
-
-		// Translate the input position to OVR coordinates
-		FVector OvrPosition = ToOVRVector(Position);
-
-		// Set the source position to current audio position
-		Result = ovrAudio_SetAudioSourcePos(AudioContext, VoiceIndex, OvrPosition.X, OvrPosition.Y, OvrPosition.Z);
-		OVR_AUDIO_CHECK(Result, "Failed to set audio source position");
-
-		// Perform the processing
-		uint32 Status;
-		Result = ovrAudio_SpatializeMonoSourceInterleaved(AudioContext, VoiceIndex, ovrAudioSpatializationFlag_None, &Status, OutSamples, InSamples);
-		OVR_AUDIO_CHECK(Result, "Failed to spatialize mono source interleaved");
-	}
-
-	/** Helper function to convert from UE coords to OVR coords. */
-	FORCEINLINE FVector ToOVRVector(const FVector& InVec) const
-	{
-		return FVector(float(InVec.Y), float(-InVec.Z), float(-InVec.X));
-	}
-
-	/* Whether or not the OVR audio context is initialized. We defer initialization until the first audio callback.*/
-	bool bOvrContextInitialized;
-
-	/** Ptr to parent audio device */
-	FAudioDevice* AudioDevice;
-
-	/** The OVR Audio context initialized to "Fast" algorithm. */
-	ovrAudioContext OvrAudioContext;
-
-	/** Xaudio2 effects for the oculus plugin */
-	TArray<class FXAudio2HRTFEffect*> HRTFEffects;
-	TArray<FSpatializationParams> Params;
-
-	FCriticalSection ParamCriticalSection;
-};
-
-/**
-* FHrtfSpatializationAlgorithmAudioMixer
-* Simple implementation of the IAudioSpatializationAlgorithm interface for use with audio mixer.
-*/
-
-class FHrtfSpatializationAlgorithmAudioMixer : public IAudioSpatialization
-{
-public:
-	FHrtfSpatializationAlgorithmAudioMixer(class FAudioDevice * AudioDevice)
-		: bOvrContextInitialized(false)
-		, AudioDevice(AudioDevice)
-		, OvrAudioContext(nullptr)
-	{
-		uint32 MaxNumSources = AudioDevice->MaxChannels;
-		int32 SampleRate = (int32)AudioDevice->GetSampleRate();
-		int32 BufferLength = AudioDevice->GetBufferLength();
-
-		check(MaxNumSources > 0);
-		check(SampleRate > 0);
-		check(BufferLength > 0);
-
-		Params.AddDefaulted(MaxNumSources);
-
-		ovrAudioContextConfiguration ContextConfig;
-		ContextConfig.acc_Size = sizeof(ovrAudioContextConfiguration);
-
-		// First initialize the Fast algorithm context
-		ContextConfig.acc_Provider = ovrAudioSpatializationProvider_OVR_OculusHQ;
-		ContextConfig.acc_MaxNumSources = MaxNumSources;
-		ContextConfig.acc_SampleRate = SampleRate;
-		ContextConfig.acc_BufferLength = (uint32)BufferLength;
-
-		check(OvrAudioContext == nullptr);
-		// Create the OVR Audio Context with a given quality
-		ovrResult Result = ovrAudio_CreateContext(&OvrAudioContext, &ContextConfig);
-		OVR_AUDIO_CHECK(Result, "Failed to create simple context");
-
-		// Now initialize the high quality algorithm context
-		bOvrContextInitialized = true;
-	}
-
-	bool IsSpatializationEffectInitialized() const override
-	{
-		return bOvrContextInitialized;
-	}
-
-	void InitializeSpatializationEffect(uint32 BufferLength) override
-	{
-		if (bOvrContextInitialized)
-		{
-			return;
-		}
-
-		uint32 MaxNumSources = AudioDevice->MaxChannels;
-		uint32 SampleRate = AudioDevice->SampleRate;
-
-		ovrAudioContextConfiguration ContextConfig;
-		ContextConfig.acc_Size = sizeof(ovrAudioContextConfiguration);
-
-		// First initialize the Fast algorithm context
-		ContextConfig.acc_Provider = ovrAudioSpatializationProvider_OVR_OculusHQ;
-		ContextConfig.acc_MaxNumSources = MaxNumSources;
-		ContextConfig.acc_SampleRate = SampleRate;
-		ContextConfig.acc_BufferLength = BufferLength;
-
-		check(OvrAudioContext == nullptr);
-		// Create the OVR Audio Context with a given quality
-		ovrResult Result = ovrAudio_CreateContext(&OvrAudioContext, &ContextConfig);
-		OVR_AUDIO_CHECK(Result, "Failed to create simple context");
-
-		for (int32 SourceId = 0; SourceId < (int32)MaxNumSources; ++SourceId)
-		{
-			Result = ovrAudio_SetAudioSourceAttenuationMode(OvrAudioContext, SourceId, ovrAudioSourceAttenuationMode_None, 1.0f);
-			OVR_AUDIO_CHECK(Result, "Failed to set source attenuation mode");
-		}
-
-		// Now initialize the high quality algorithm context
-		bOvrContextInitialized = true;
-	}
-
-	virtual void SetSpatializationParameters(uint32 VoiceId, const FSpatializationParams& InParams)
-	{
-		Params[VoiceId] = InParams;
-	}
-
-	void ProcessAudio(const FAudioPluginSourceInputData& InputData, FAudioPluginSourceOutputData& OutputData) override
-	{
-		if (InputData.SpatializationParams && OvrAudioContext)
-		{
-			Params[InputData.SourceId] = *InputData.SpatializationParams;
-
-			// Translate the input position to OVR coordinates
-			FVector OvrPosition = ToOVRVector(Params[InputData.SourceId].EmitterPosition);
-
-			// Set the source position to current audio position
-			ovrResult Result = ovrAudio_SetAudioSourcePos(OvrAudioContext, InputData.SourceId, OvrPosition.X, OvrPosition.Y, OvrPosition.Z);
-			OVR_AUDIO_CHECK(Result, "Failed to set audio source position");
-
-			// Perform the processing
-			uint32 Status;
-			Result = ovrAudio_SpatializeMonoSourceInterleaved(OvrAudioContext, InputData.SourceId, ovrAudioSpatializationFlag_None, &Status, OutputData.AudioBuffer.GetData(), InputData.AudioBuffer->GetData());
-			OVR_AUDIO_CHECK(Result, "Failed to spatialize mono source interleaved");
-		}
-	}
-
-private:
-
-	/** Helper function to convert from UE coords to OVR coords. */
-	FORCEINLINE FVector ToOVRVector(const FVector& InVec) const
-	{
-		return FVector(InVec.Y, -InVec.Z, -InVec.X);
-	}
 
-	/* Whether or not the OVR audio context is initialized. We defer initialization until the first audio callback.*/
-	bool bOvrContextInitialized;
-
-	TArray<FSpatializationParams> Params;
-
-	/** Ptr to parent audio device */
-	FAudioDevice* AudioDevice;
-
-	/** The OVR Audio context initialized to "Fast" algorithm. */
-	ovrAudioContext OvrAudioContext;
-};
-
-/**
-* FOculusAudioPlugin
-*/
-
-class FOculusAudioPlugin : public IOculusAudioPlugin
+TAudioSpatializationPtr FOculusSpatializationPluginFactory::CreateNewSpatializationPlugin(FAudioDevice* OwningDevice)
 {
-public:
-	FOculusAudioPlugin()
+	if (OwningDevice->IsAudioMixerEnabled())
 	{
+		return TAudioSpatializationPtr(new OculusAudioSpatializationAudioMixer());
 	}
-
-	void Initialize() override
+	else
 	{
-		if (FOculusAudioPlugin::NumInstances == 0)
-		{
-			if (!LoadModule())
-			{
-				UE_LOG(LogAudio, Error, TEXT("Failed to load OVR Audio dll"));
-				return;
-			}
-		}
-
-		FOculusAudioPlugin::NumInstances++;
-
-		if (!FOculusAudioPlugin::bInitialized)
-		{
-			// Check the version number
-			int32 MajorVersionNumber;
-			int32 MinorVersionNumber;
-			int32 PatchNumber;
-
-			// Initialize the OVR Audio SDK before making any calls to ovrAudio
-			ovrResult Result = ovrAudio_Initialize();
-			OVR_AUDIO_CHECK(Result, "Failed to initialize OVR Audio system");
-
-			const char* OvrVersionString = ovrAudio_GetVersion(&MajorVersionNumber, &MinorVersionNumber, &PatchNumber);
-			if (MajorVersionNumber != OVR_AUDIO_MAJOR_VERSION || MinorVersionNumber != OVR_AUDIO_MINOR_VERSION)
-			{
-				UE_LOG(LogAudio, Warning, TEXT("Using mismatched OVR Audio SDK Versiont! %d.%d vs. %d.%d"), OVR_AUDIO_MAJOR_VERSION, OVR_AUDIO_MINOR_VERSION, MajorVersionNumber, MinorVersionNumber);
-				return;
-			}
-			FOculusAudioPlugin::bInitialized = true;
-		}
-	}
-
-	void Shutdown() override
-	{
-		if (FOculusAudioPlugin::NumInstances == 0)
-		{
-			// This means we failed to load the OVR Audio module during initialization and there's nothing to shutdown.
-			return;
-		}
-
-		FOculusAudioPlugin::NumInstances--;
-
-		if (FOculusAudioPlugin::NumInstances == 0)
-		{
-			// Shutdown OVR audio
-			ovrAudio_Shutdown();
-
-			ShutdownModule();
-
-			FOculusAudioPlugin::bInitialized = false;
-		}
+		return TAudioSpatializationPtr(new OculusAudioLegacySpatialization());
 	}
-
-	TAudioSpatializationPtr CreateSpatializationInterface(class FAudioDevice* AudioDevice) override
-	{
-		// If we are using the audio mixer, use a new implementation
-		if (AudioDevice->IsAudioMixerEnabled())
-		{
-			return TAudioSpatializationPtr(new FHrtfSpatializationAlgorithmAudioMixer(AudioDevice));
-		}
-
-		// Otherwise, use the xaudio2 implementation
-		return TAudioSpatializationPtr(new FHrtfSpatializationAlgorithmXAudio2(AudioDevice));
-	}
-
-private:
-
-	bool LoadModule()
-	{
-		if (!FOculusAudioPlugin::OculusAudioDllHandle)
-		{
-			FString Path = FPaths::EngineDir() / FString::Printf(TEXT("Binaries/ThirdParty/Oculus/Audio/Win64/"));
-			FPlatformProcess::PushDllDirectory(*Path);
-			FOculusAudioPlugin::OculusAudioDllHandle = FPlatformProcess::GetDllHandle(*(Path + "ovraudio64.dll"));
-			FPlatformProcess::PopDllDirectory(*Path);
-			return FOculusAudioPlugin::OculusAudioDllHandle != nullptr;
-		}
-		return true;
-	}
-
-	void ReleaseModule()
-	{
-
-		if (FOculusAudioPlugin::NumInstances == 0 && FOculusAudioPlugin::OculusAudioDllHandle)
-		{
-			FPlatformProcess::FreeDllHandle(OculusAudioDllHandle);
-			FOculusAudioPlugin::OculusAudioDllHandle = nullptr;
-		}
-	}
-
-	static void* OculusAudioDllHandle;
-	static uint32 NumInstances;
-	static bool bInitialized;
-};
-
-void* FOculusAudioPlugin::OculusAudioDllHandle = nullptr;
-uint32 FOculusAudioPlugin::NumInstances = 0;
-bool FOculusAudioPlugin::bInitialized = false;
-
-IMPLEMENT_MODULE( FOculusAudioPlugin, OculusAudio )
+}
diff --git a/Engine/Plugins/Runtime/Oculus/OculusAudio/Source/OculusAudio/Private/OculusAudio.h b/Engine/Plugins/Runtime/Oculus/OculusAudio/Source/OculusAudio/Private/OculusAudio.h
--- a/Engine/Plugins/Runtime/Oculus/OculusAudio/Source/OculusAudio/Private/OculusAudio.h
+++ b/Engine/Plugins/Runtime/Oculus/OculusAudio/Source/OculusAudio/Private/OculusAudio.h
@@ -2,4 +2,36 @@
 
 #pragma once
 
-#include "IOculusAudioPlugin.h"
\ No newline at end of file
+#include "IAudioExtensionPlugin.h"
+#include "OculusAudioMixer.h"
+#include "OculusAudioLegacy.h"
+
+/************************************************************************/
+/* FOculusSpatializationPluginFactory                                   */
+/* Handles initialization of the required Oculus Audio Spatialization   */
+/* plugin.                                                              */
+/************************************************************************/
+class FOculusSpatializationPluginFactory : public IAudioSpatializationFactory
+{
+	//~ Begin IAudioSpatializationFactory
+	virtual FString GetDisplayName() override
+	{
+		static FString DisplayName = FString(TEXT("Oculus Audio"));
+		return DisplayName;
+	}
+
+	virtual bool SupportsPlatform(EAudioPlatform Platform) override
+	{
+		if (Platform == EAudioPlatform::Windows)
+		{
+			return true;
+		}
+		else
+		{
+			return false;
+		}
+	}
+
+	virtual TAudioSpatializationPtr CreateNewSpatializationPlugin(FAudioDevice* OwningDevice) override;
+	//~ End IAudioSpatializationFactory
+};
diff --git a/Engine/Plugins/Runtime/Oculus/OculusAudio/Source/OculusAudio/Private/OculusAudioDllManager.cpp b/Engine/Plugins/Runtime/Oculus/OculusAudio/Source/OculusAudio/Private/OculusAudioDllManager.cpp
new file mode 100644
--- /dev/null
+++ b/Engine/Plugins/Runtime/Oculus/OculusAudio/Source/OculusAudio/Private/OculusAudioDllManager.cpp
@@ -0,0 +1,82 @@
+// Copyright 1998-2017 Epic Games, Inc. All Rights Reserved.
+#include "OculusAudioDllManager.h"
+#include "Paths.h"
+
+void* FOculusAudioLibraryManager::OculusAudioDllHandle = nullptr;
+uint32 FOculusAudioLibraryManager::NumInstances = 0;
+bool FOculusAudioLibraryManager::bInitialized = false;
+
+void FOculusAudioLibraryManager::Initialize()
+{
+	if (FOculusAudioLibraryManager::NumInstances == 0)
+	{
+		if (!LoadDll())
+		{
+			UE_LOG(LogAudio, Error, TEXT("Failed to load OVR Audio dll"));
+			check(false);
+			return;
+		}
+	}
+
+	FOculusAudioLibraryManager::NumInstances++;
+	if (!FOculusAudioLibraryManager::bInitialized)
+	{
+		// Check the version number
+		int32 MajorVersionNumber;
+		int32 MinorVersionNumber;
+		int32 PatchNumber;
+
+		// Initialize the OVR Audio SDK before making any calls to ovrAudio
+		ovrResult Result = ovrAudio_Initialize();
+		OVR_AUDIO_CHECK(Result, "Failed to initialize OVR Audio system");
+
+		const char* OvrVersionString = ovrAudio_GetVersion(&MajorVersionNumber, &MinorVersionNumber, &PatchNumber);
+		if (MajorVersionNumber != OVR_AUDIO_MAJOR_VERSION || MinorVersionNumber != OVR_AUDIO_MINOR_VERSION)
+		{
+			UE_LOG(LogAudio, Warning, TEXT("Using mismatched OVR Audio SDK Version! %d.%d vs. %d.%d"), OVR_AUDIO_MAJOR_VERSION, OVR_AUDIO_MINOR_VERSION, MajorVersionNumber, MinorVersionNumber);
+			return;
+		}
+		FOculusAudioLibraryManager::bInitialized = true;
+	}
+}
+
+void FOculusAudioLibraryManager::Shutdown()
+{
+	if (FOculusAudioLibraryManager::NumInstances == 0)
+	{
+		// This means we failed to load the OVR Audio module during initialization and there's nothing to shutdown.
+		return;
+	}
+
+	FOculusAudioLibraryManager::NumInstances--;
+
+	if (FOculusAudioLibraryManager::NumInstances == 0)
+	{
+		// Shutdown OVR audio
+		ovrAudio_Shutdown();
+		ReleaseDll();
+		FOculusAudioLibraryManager::bInitialized = false;
+	}
+}
+
+bool FOculusAudioLibraryManager::LoadDll()
+{
+	if (!FOculusAudioLibraryManager::OculusAudioDllHandle)
+	{
+		FString Path = FPaths::EngineDir() / FString::Printf(TEXT("Binaries/ThirdParty/Oculus/Audio/Win64/"));
+		FPlatformProcess::PushDllDirectory(*Path);
+		FOculusAudioLibraryManager::OculusAudioDllHandle = FPlatformProcess::GetDllHandle(*(Path + "ovraudio64.dll"));
+		FPlatformProcess::PopDllDirectory(*Path);
+		return FOculusAudioLibraryManager::OculusAudioDllHandle != nullptr;
+	}
+	return true;
+}
+
+void FOculusAudioLibraryManager::ReleaseDll()
+{
+	if (FOculusAudioLibraryManager::NumInstances == 0 && FOculusAudioLibraryManager::OculusAudioDllHandle)
+	{
+		FPlatformProcess::FreeDllHandle(OculusAudioDllHandle);
+		FOculusAudioLibraryManager::OculusAudioDllHandle = nullptr;
+	}
+}
\ No newline at end of file
diff --git a/Engine/Plugins/Runtime/Oculus/OculusAudio/Source/OculusAudio/Private/OculusAudioDllManager.h b/Engine/Plugins/Runtime/Oculus/OculusAudio/Source/OculusAudio/Private/OculusAudioDllManager.h
new file mode 100644
--- /dev/null
+++ b/Engine/Plugins/Runtime/Oculus/OculusAudio/Source/OculusAudio/Private/OculusAudioDllManager.h
@@ -0,0 +1,64 @@
+// Copyright 1998-2017 Epic Games, Inc. All Rights Reserved.
+#pragma once
+#include "OVR_Audio.h"
+#include "CoreMinimal.h"
+#include "Audio.h"
+
+//---------------------------------------------------
+// Oculus Audio DLL handling
+//---------------------------------------------------
+
+static const TCHAR* GetOculusErrorString(ovrResult Result)
+{
+	switch (Result)
+	{
+		default:
+		case ovrError_AudioUnknown:			return TEXT("Unknown Error");
+		case ovrError_AudioInvalidParam:	return TEXT("Invalid Param");
+		case ovrError_AudioBadSampleRate:	return TEXT("Bad Samplerate");
+		case ovrError_AudioMissingDLL:		return TEXT("Missing DLL");
+		case ovrError_AudioBadAlignment:	return TEXT("Pointers did not meet 16 byte alignment requirements");
+		case ovrError_AudioUninitialized:	return TEXT("Function called before initialization");
+		case ovrError_AudioHRTFInitFailure: return TEXT("HRTF Profider initialization failed");
+		case ovrError_AudioBadVersion:		return TEXT("Bad audio version");
+		case ovrError_AudioSRBegin:			return TEXT("Sample rate begin");
+		case ovrError_AudioSREnd:			return TEXT("Sample rate end");
+	}
+}
+
+#define OVR_AUDIO_CHECK(Result, Context)																\
+	if (Result != ovrSuccess)																			\
+	{																									\
+		const TCHAR* ErrString = GetOculusErrorString(Result);											\
+		UE_LOG(LogAudio, Error, TEXT("Oculus Audio SDK Error - %s: %s"), TEXT(Context), ErrString);	\
+		return;																							\
+	}
+
+
+/************************************************************************/
+/* FOculusAudioLibraryManager                                           */
+/* This handles loading and unloading the Oculus Audio DLL at runtime.  */
+/************************************************************************/
+class FOculusAudioLibraryManager
+{
+public:
+	static bool IsInitialized()
+	{
+		return bInitialized;
+	}
+
+	static void Initialize();
+	
+
+	static void Shutdown();
+
+private:
+
+	static bool LoadDll();
+
+	static void ReleaseDll();
+
+	static void* OculusAudioDllHandle;
+	static uint32 NumInstances;
+	static bool bInitialized;
+};
\ No newline at end of file
diff --git a/Engine/Plugins/Runtime/Oculus/OculusAudio/Source/OculusAudio/Private/OculusAudioEffect.cpp b/Engine/Plugins/Runtime/Oculus/OculusAudio/Source/OculusAudio/Private/OculusAudioEffect.cpp
--- a/Engine/Plugins/Runtime/Oculus/OculusAudio/Source/OculusAudio/Private/OculusAudioEffect.cpp
+++ b/Engine/Plugins/Runtime/Oculus/OculusAudio/Source/OculusAudio/Private/OculusAudioEffect.cpp
@@ -87,10 +87,6 @@ void FXAudio2HRTFEffect::Process(UINT32 InputProcessParameterCount, const XAPO_P
 
 	FVector& EmitterPosition = CurrentParameters.EmitterPosition;
 
-	// We better not have default set at this point because default means "no HRTF processing". We also
-	// better thave the plugin because this XAPO code is inside the plugin
-	check(AudioDevice->AudioPlugin != nullptr);
-
 	XAPO_BUFFER_FLAGS InFlags = pInputProcessParameters[0].BufferFlags;
 	XAPO_BUFFER_FLAGS OutFlags = pOutputProcessParameters[0].BufferFlags;
 
diff --git a/Engine/Plugins/Runtime/Oculus/OculusAudio/Source/OculusAudio/Private/OculusAudioLegacy.cpp b/Engine/Plugins/Runtime/Oculus/OculusAudio/Source/OculusAudio/Private/OculusAudioLegacy.cpp
new file mode 100644
--- /dev/null
+++ b/Engine/Plugins/Runtime/Oculus/OculusAudio/Source/OculusAudio/Private/OculusAudioLegacy.cpp
@@ -0,0 +1,151 @@
+// Copyright 1998-2017 Epic Games, Inc. All Rights Reserved.
+
+#include "OculusAudioLegacy.h"
+
+OculusAudioLegacySpatialization::OculusAudioLegacySpatialization()
+	: bOvrContextInitialized(false)
+	, OvrAudioContext(nullptr)
+{
+	
+}
+
+OculusAudioLegacySpatialization::~OculusAudioLegacySpatialization()
+{
+}
+
+void OculusAudioLegacySpatialization::Initialize(const FAudioPluginInitializationParams InitializationParams)
+{
+	check(HRTFEffects.Num() == 0);
+	check(Params.Num() == 0);
+
+	for (int32 EffectIndex = 0; EffectIndex < (int32)InitializationParams.NumSources; ++EffectIndex)
+	{
+		/* Hack: grab main audio device. */
+		FXAudio2HRTFEffect* NewHRTFEffect = new FXAudio2HRTFEffect(EffectIndex, InitializationParams.AudioDevicePtr);
+		/* End hack. */
+		NewHRTFEffect->Initialize(nullptr, 0);
+		HRTFEffects.Add(NewHRTFEffect);
+
+		Params.Add(FSpatializationParams());
+	}
+
+	if (bOvrContextInitialized)
+	{
+		return;
+	}
+
+	ovrAudioContextConfiguration ContextConfig;
+	ContextConfig.acc_Size = sizeof(ovrAudioContextConfiguration);
+
+	// First initialize the Fast algorithm context
+	ContextConfig.acc_Provider = ovrAudioSpatializationProvider_OVR_OculusHQ;
+	ContextConfig.acc_MaxNumSources = InitializationParams.NumSources;
+	ContextConfig.acc_SampleRate = InitializationParams.SampleRate;
+	//XAudio2 sets the buffer callback size to a 100th of the sample rate:
+	ContextConfig.acc_BufferLength = InitializationParams.SampleRate / 100;
+
+	check(OvrAudioContext == nullptr);
+	// Create the OVR Audio Context with a given quality
+	ovrResult Result = ovrAudio_CreateContext(&OvrAudioContext, &ContextConfig);
+	OVR_AUDIO_CHECK(Result, "Failed to create simple context");
+
+	// Now initialize the high quality algorithm context
+	bOvrContextInitialized = true;
+}
+
+void OculusAudioLegacySpatialization::Shutdown()
+{
+	// Release all the effects for the oculus spatialization effect
+	for (FXAudio2HRTFEffect* Effect : HRTFEffects)
+	{
+		if (Effect)
+		{
+			delete Effect;
+		}
+	}
+	HRTFEffects.Empty();
+
+	// Destroy the contexts if we created them
+	if (bOvrContextInitialized)
+	{
+		if (OvrAudioContext != nullptr)
+		{
+			ovrAudio_DestroyContext(OvrAudioContext);
+			OvrAudioContext = nullptr;
+		}
+
+		bOvrContextInitialized = false;
+	}
+}
+
+bool OculusAudioLegacySpatialization::IsSpatializationEffectInitialized() const
+{
+	return bOvrContextInitialized;
+}
+
+void OculusAudioLegacySpatialization::ProcessSpatializationForVoice(uint32 VoiceIndex, float* InSamples, float* OutSamples, const FVector& Position)
+{
+	if (OvrAudioContext)
+	{
+		ProcessAudioInternal(OvrAudioContext, VoiceIndex, InSamples, OutSamples, Position);
+	}
+}
+
+bool OculusAudioLegacySpatialization::CreateSpatializationEffect(uint32 VoiceId)
+{
+	// If an effect for this voice has already been created, then leave
+	if ((int32)VoiceId >= HRTFEffects.Num())
+	{
+		return false;
+	}
+	return true;
+}
+
+void* OculusAudioLegacySpatialization::GetSpatializationEffect(uint32 VoiceId)
+{
+	if ((int32)VoiceId < HRTFEffects.Num())
+	{
+		return (void*)HRTFEffects[VoiceId];
+	}
+	return nullptr;
+}
+
+void OculusAudioLegacySpatialization::SetSpatializationParameters(uint32 VoiceId, const FSpatializationParams& InParams)
+{
+	check((int32)VoiceId < Params.Num());
+
+	FScopeLock ScopeLock(&ParamCriticalSection);
+	Params[VoiceId] = InParams;
+}
+
+void OculusAudioLegacySpatialization::GetSpatializationParameters(uint32 VoiceId, FSpatializationParams& OutParams)
+{
+	check((int32)VoiceId < Params.Num());
+
+	FScopeLock ScopeLock(&ParamCriticalSection);
+	OutParams = Params[VoiceId];
+}
+
+void OculusAudioLegacySpatialization::ProcessAudioInternal(ovrAudioContext AudioContext, uint32 VoiceIndex, float* InSamples, float* OutSamples, const FVector& Position)
+{
+	ovrResult Result = ovrAudio_SetAudioSourceAttenuationMode(AudioContext, VoiceIndex, ovrAudioSourceAttenuationMode_None, 1.0f);
+	OVR_AUDIO_CHECK(Result, "Failed to set source attenuation mode");
+
+	// Translate the input position to OVR coordinates
+	FVector OvrPosition = ToOVRVector(Position);
+
+	// Set the source position to current audio position
+	Result = ovrAudio_SetAudioSourcePos(AudioContext, VoiceIndex, OvrPosition.X, OvrPosition.Y, OvrPosition.Z);
+	OVR_AUDIO_CHECK(Result, "Failed to set audio source position");
+
+	// Perform the processing
+	uint32 Status;
+	Result = ovrAudio_SpatializeMonoSourceInterleaved(AudioContext, VoiceIndex, ovrAudioSpatializationFlag_None, &Status, OutSamples, InSamples);
+	OVR_AUDIO_CHECK(Result, "Failed to spatialize mono source interleaved");
+}
+
+/** Helper function to convert from UE coords to OVR coords. */
+FVector OculusAudioLegacySpatialization::ToOVRVector(const FVector& InVec) const
+{
+	return FVector(float(InVec.Y), float(-InVec.Z), float(-InVec.X));
+}
\ No newline at end of file
diff --git a/Engine/Plugins/Runtime/Oculus/OculusAudio/Source/OculusAudio/Private/OculusAudioLegacy.h b/Engine/Plugins/Runtime/Oculus/OculusAudio/Source/OculusAudio/Private/OculusAudioLegacy.h
new file mode 100644
--- /dev/null
+++ b/Engine/Plugins/Runtime/Oculus/OculusAudio/Source/OculusAudio/Private/OculusAudioLegacy.h
@@ -0,0 +1,52 @@
+// Copyright 1998-2017 Epic Games, Inc. All Rights Reserved.
+
+#pragma once
+#include "OVR_Audio.h"
+#include "XAudio2Device.h"
+#include "AudioEffect.h"
+#include "XAudio2Effects.h"
+#include "OculusAudioEffect.h"
+#include "Runtime/Windows/XAudio2/Private/XAudio2Support.h"
+#include "Misc/Paths.h"
+#include "OculusAudioDllManager.h"
+
+/************************************************************************/
+/* OculusAudioLegacySpatialization                                      */
+/* This spatialization plugin is used in the non-audiomixer engine,     */
+/* XAudio2's HRTFEffect plugin system directly.                         */
+/************************************************************************/
+class OculusAudioLegacySpatialization : public IAudioSpatialization
+{
+public:
+	OculusAudioLegacySpatialization();
+	~OculusAudioLegacySpatialization();
+
+	/** Begin IAudioSpatialization implementation */
+	virtual void Initialize(const FAudioPluginInitializationParams InitializationParams) override;
+	virtual void Shutdown() override;
+
+	virtual bool IsSpatializationEffectInitialized() const override;
+	virtual void ProcessSpatializationForVoice(uint32 VoiceIndex, float* InSamples, float* OutSamples, const FVector& Position) override;
+	virtual bool CreateSpatializationEffect(uint32 VoiceId) override;
+	virtual void* GetSpatializationEffect(uint32 VoiceId) override;
+	virtual void SetSpatializationParameters(uint32 VoiceId, const FSpatializationParams& InParams) override;
+	virtual void GetSpatializationParameters(uint32 VoiceId, FSpatializationParams& OutParams) override;
+	/** End IAudioSpatialization implementation */
+
+private:
+	void ProcessAudioInternal(ovrAudioContext AudioContext, uint32 VoiceIndex, float* InSamples, float* OutSamples, const FVector& Position);
+	FVector ToOVRVector(const FVector& InVec) const;
+	
+private:
+	/* Whether or not the OVR audio context is initialized. We defer initialization until the first audio callback.*/
+	bool bOvrContextInitialized;
+
+	/** The OVR Audio context initialized to "Fast" algorithm. */
+	ovrAudioContext OvrAudioContext;
+
+	/** Xaudio2 effects for the oculus plugin */
+	TArray<class FXAudio2HRTFEffect*> HRTFEffects;
+	TArray<FSpatializationParams> Params;
+
+	FCriticalSection ParamCriticalSection;
+};
\ No newline at end of file
diff --git a/Engine/Plugins/Runtime/Oculus/OculusAudio/Source/OculusAudio/Private/OculusAudioMixer.cpp b/Engine/Plugins/Runtime/Oculus/OculusAudio/Source/OculusAudio/Private/OculusAudioMixer.cpp
new file mode 100644
--- /dev/null
+++ b/Engine/Plugins/Runtime/Oculus/OculusAudio/Source/OculusAudio/Private/OculusAudioMixer.cpp
@@ -0,0 +1,86 @@
+// Copyright 1998-2017 Epic Games, Inc. All Rights Reserved.
+
+#include "OculusAudioMixer.h"
+
+OculusAudioSpatializationAudioMixer::OculusAudioSpatializationAudioMixer()
+	: bOvrContextInitialized(false)
+	, OvrAudioContext(nullptr)
+{
+}
+
+OculusAudioSpatializationAudioMixer::~OculusAudioSpatializationAudioMixer()
+{
+	if (bOvrContextInitialized)
+	{
+		Shutdown();
+	}
+}
+
+void OculusAudioSpatializationAudioMixer::Initialize(const FAudioPluginInitializationParams InitializationParams)
+{
+	if (bOvrContextInitialized)
+	{
+		return;
+	}
+
+	Params.AddDefaulted(InitializationParams.NumSources);
+
+	ovrAudioContextConfiguration ContextConfig;
+	ContextConfig.acc_Size = sizeof(ovrAudioContextConfiguration);
+
+	// First initialize the Fast algorithm context
+	ContextConfig.acc_Provider = ovrAudioSpatializationProvider_OVR_OculusHQ;
+	ContextConfig.acc_MaxNumSources = InitializationParams.NumSources;
+	ContextConfig.acc_SampleRate = InitializationParams.SampleRate;
+	ContextConfig.acc_BufferLength = InitializationParams.BufferLength;
+
+	check(OvrAudioContext == nullptr);
+	// Create the OVR Audio Context with a given quality
+	ovrResult Result = ovrAudio_CreateContext(&OvrAudioContext, &ContextConfig);
+	OVR_AUDIO_CHECK(Result, "Failed to create simple context");
+
+	for (int32 SourceId = 0; SourceId < (int32)InitializationParams.NumSources; ++SourceId)
+	{
+		Result = ovrAudio_SetAudioSourceAttenuationMode(OvrAudioContext, SourceId, ovrAudioSourceAttenuationMode_None, 1.0f);
+		OVR_AUDIO_CHECK(Result, "Failed to set source attenuation mode");
+	}
+
+	// Now initialize the high quality algorithm context
+	bOvrContextInitialized = true;
+}
+
+void OculusAudioSpatializationAudioMixer::Shutdown()
+{
+	ovrAudio_DestroyContext(OvrAudioContext);
+	bOvrContextInitialized = false;
+}
+
+bool OculusAudioSpatializationAudioMixer::IsSpatializationEffectInitialized() const
+{
+	return bOvrContextInitialized;
+}
+
+void OculusAudioSpatializationAudioMixer::SetSpatializationParameters(uint32 VoiceId, const FSpatializationParams& InParams)
+{
+	Params[VoiceId] = InParams;
+}
+
+void OculusAudioSpatializationAudioMixer::ProcessAudio(const FAudioPluginSourceInputData& InputData, FAudioPluginSourceOutputData& OutputData)
+{
+	if (InputData.SpatializationParams && OvrAudioContext)
+	{
+		Params[InputData.SourceId] = *InputData.SpatializationParams;
+
+		// Translate the input position to OVR coordinates
+		FVector OvrPosition = ToOVRVector(Params[InputData.SourceId].EmitterPosition);
+
+		// Set the source position to current audio position
+		ovrResult Result = ovrAudio_SetAudioSourcePos(OvrAudioContext, InputData.SourceId, OvrPosition.X, OvrPosition.Y, OvrPosition.Z);
+		OVR_AUDIO_CHECK(Result, "Failed to set audio source position");
+
+		// Perform the processing
+		uint32 Status;
+		Result = ovrAudio_SpatializeMonoSourceInterleaved(OvrAudioContext, InputData.SourceId, ovrAudioSpatializationFlag_None, &Status, OutputData.AudioBuffer.GetData(), InputData.AudioBuffer->GetData());
+		OVR_AUDIO_CHECK(Result, "Failed to spatialize mono source interleaved");
+	}
+}
\ No newline at end of file
diff --git a/Engine/Plugins/Runtime/Oculus/OculusAudio/Source/OculusAudio/Private/OculusAudioMixer.h b/Engine/Plugins/Runtime/Oculus/OculusAudio/Source/OculusAudio/Private/OculusAudioMixer.h
new file mode 100644
--- /dev/null
+++ b/Engine/Plugins/Runtime/Oculus/OculusAudio/Source/OculusAudio/Private/OculusAudioMixer.h
@@ -0,0 +1,46 @@
+// Copyright 1998-2017 Epic Games, Inc. All Rights Reserved.
+
+#pragma once
+#include "OVR_Audio.h"
+#include "XAudio2Device.h"
+#include "AudioEffect.h"
+#include "XAudio2Effects.h"
+#include "OculusAudioEffect.h"
+#include "Runtime/Windows/XAudio2/Private/XAudio2Support.h"
+#include "Misc/Paths.h"
+#include "OculusAudioDllManager.h"
+
+/************************************************************************/
+/* OculusAudioSpatializationAudioMixer                                  */
+/* This implementation of IAudioSpatialization uses the Oculus Audio    */
+/* library to render audio sources with HRTF spatialization.            */
+/************************************************************************/
+class OculusAudioSpatializationAudioMixer : public IAudioSpatialization
+{
+public:
+	OculusAudioSpatializationAudioMixer();
+	~OculusAudioSpatializationAudioMixer();
+
+	//~ Begin IAudioSpatialization 
+	virtual void Initialize(const FAudioPluginInitializationParams InitializationParams) override;
+	virtual void Shutdown() override;
+	virtual bool IsSpatializationEffectInitialized() const override;
+	virtual void SetSpatializationParameters(uint32 VoiceId, const FSpatializationParams& InParams) override;
+	virtual void ProcessAudio(const FAudioPluginSourceInputData& InputData, FAudioPluginSourceOutputData& OutputData) override;
+	//~ End IAudioSpatialization 
+
+private:
+	// Helper function to convert from UE coords to OVR coords.
+	FORCEINLINE FVector ToOVRVector(const FVector& InVec) const
+	{
+		return FVector(InVec.Y, -InVec.Z, -InVec.X);
+	}
+
+	// Whether or not the OVR audio context is initialized. We defer initialization until the first audio callback.
+	bool bOvrContextInitialized;
+
+	TArray<FSpatializationParams> Params;
+
+	// The OVR Audio context initialized to "Fast" algorithm.
+	ovrAudioContext OvrAudioContext;
+};
\ No newline at end of file
diff --git a/Engine/Plugins/Runtime/Oculus/OculusAudio/Source/OculusAudio/Public/IOculusAudioPlugin.h b/Engine/Plugins/Runtime/Oculus/OculusAudio/Source/OculusAudio/Public/IOculusAudioPlugin.h
--- a/Engine/Plugins/Runtime/Oculus/OculusAudio/Source/OculusAudio/Public/IOculusAudioPlugin.h
+++ b/Engine/Plugins/Runtime/Oculus/OculusAudio/Source/OculusAudio/Public/IOculusAudioPlugin.h
@@ -3,39 +3,22 @@
 #pragma once
 
 #include "ModuleManager.h"
-#include "IAudioExtensionPlugin.h"
-
-/**
- * The public interface to this module.  In most cases, this interface is only public to sibling modules 
- * within this plugin.
- */
-class IOculusAudioPlugin : public IAudioPlugin
+#include "OculusAudio.h"
+#include "OculusAudioDllManager.h"
+
+/************************************************************************/
+/* FOculusAudioPlugin                                                   */
+/* Module interface. Also handles loading and unloading the Oculus Audio*/
+/* DLL using FOculusAudioDllManager.                                    */
+/************************************************************************/
+class FOculusAudioPlugin : public IModuleInterface
 {
-
 public:
+	//~ Begin IModuleInterface 
+	virtual void StartupModule() override;
+	virtual void ShutdownModule() override;
+	//~ End IModuleInterface
 
-	/**
-	 * Singleton-like access to this module's interface.  This is just for convenience!
-	 * Beware of calling this during the shutdown phase, though.  Your module might have been unloaded already.
-	 *
-	 * @return Returns singleton instance, loading the module on demand if needed
-	 */
-	static inline IOculusAudioPlugin& Get()
-	{
-		return FModuleManager::LoadModuleChecked< IOculusAudioPlugin >("OculusAudio");
-	}
-
-	/**
-	 * Checks to see if this module is loaded and ready.  It is only valid to call Get() if IsAvailable() returns true.
-	 *
-	 * @return True if the module is loaded and ready to use
-	 */
-	static inline bool IsAvailable()
-	{
-		return FModuleManager::Get().IsModuleLoaded( "OculusAudio" );
-	}
-
-	/** We are overriding spatialization. */
-	virtual bool ImplementsSpatialization() const override { return true; }
+private:
+	FOculusSpatializationPluginFactory PluginFactory;
 };
-
diff --git a/Engine/Plugins/Runtime/PhysXVehicles/Source/PhysXVehicles/Private/WheeledVehicleMovementComponent.cpp b/Engine/Plugins/Runtime/PhysXVehicles/Source/PhysXVehicles/Private/WheeledVehicleMovementComponent.cpp
--- a/Engine/Plugins/Runtime/PhysXVehicles/Source/PhysXVehicles/Private/WheeledVehicleMovementComponent.cpp
+++ b/Engine/Plugins/Runtime/PhysXVehicles/Source/PhysXVehicles/Private/WheeledVehicleMovementComponent.cpp
@@ -80,16 +80,13 @@ void PTireShader(const void* shaderData, const PxF32 tireFriction,
 	tireLongForceMag = Output.LongForce;
 	tireLatForceMag = Output.LatForce;
 	
-	if ( /*Wheel->bDebugWheels*/true )
-	{
-		Wheel->DebugLongSlip = longSlip;
-		Wheel->DebugLatSlip = latSlip;
-		Wheel->DebugNormalizedTireLoad = normalisedTireLoad;
-		Wheel->DebugTireLoad = tireLoad;
-		Wheel->DebugWheelTorque = wheelTorque;
-		Wheel->DebugLongForce = tireLongForceMag;
-		Wheel->DebugLatForce = tireLatForceMag;
-	}
+	Wheel->DebugLongSlip = longSlip;
+	Wheel->DebugLatSlip = latSlip;
+	Wheel->DebugNormalizedTireLoad = normalisedTireLoad;
+	Wheel->DebugTireLoad = tireLoad;
+	Wheel->DebugWheelTorque = wheelTorque;
+	Wheel->DebugLongForce = tireLongForceMag;
+	Wheel->DebugLatForce = tireLatForceMag;
 }
 
 UWheeledVehicleMovementComponent::UWheeledVehicleMovementComponent(const FObjectInitializer& ObjectInitializer)
diff --git a/Engine/Plugins/Runtime/PhysXVehicles/Source/PhysXVehicles/Public/VehicleAnimInstance.h b/Engine/Plugins/Runtime/PhysXVehicles/Source/PhysXVehicles/Public/VehicleAnimInstance.h
--- a/Engine/Plugins/Runtime/PhysXVehicles/Source/PhysXVehicles/Public/VehicleAnimInstance.h
+++ b/Engine/Plugins/Runtime/PhysXVehicles/Source/PhysXVehicles/Public/VehicleAnimInstance.h
@@ -20,8 +20,6 @@ struct FWheelAnimData
 	FVector LocOffset;
 };
 
-class UWheeledVehicleMovementComponent;
-
  /** Proxy override for this UAnimInstance-derived class */
 USTRUCT()
 struct PHYSXVEHICLES_API FVehicleAnimInstanceProxy : public FAnimInstanceProxy
diff --git a/Engine/Plugins/Runtime/Steam/SteamAudio/Source/SteamAudio/Private/PhononOcclusion.cpp b/Engine/Plugins/Runtime/Steam/SteamAudio/Source/SteamAudio/Private/PhononOcclusion.cpp
--- a/Engine/Plugins/Runtime/Steam/SteamAudio/Source/SteamAudio/Private/PhononOcclusion.cpp
+++ b/Engine/Plugins/Runtime/Steam/SteamAudio/Source/SteamAudio/Private/PhononOcclusion.cpp
@@ -16,7 +16,7 @@ namespace SteamAudio
 {
 	FPhononOcclusion::FPhononOcclusion()
 		: EnvironmentalRenderer(nullptr)
-		, SteamAudioModule(nullptr)
+		, EnvironmentCriticalSectionHandle(nullptr)
 	{
 		InputAudioFormat.channelLayout = IPL_CHANNELLAYOUT_MONO;
 		InputAudioFormat.channelLayoutType = IPL_CHANNELLAYOUTTYPE_SPEAKERS;
@@ -41,24 +41,22 @@ namespace SteamAudio
 	{
 	}
 
-	void FPhononOcclusion::Initialize(const int32 SampleRate, const int32 NumSources, const int32 FrameSize)
+	void FPhononOcclusion::Initialize(const FAudioPluginInitializationParams InitializationParams)
 	{
-		DirectSoundSources.AddDefaulted(NumSources);
+		DirectSoundSources.AddDefaulted(InitializationParams.NumSources);
 
-		for (auto& DirectSoundSource : DirectSoundSources)
+		for (FDirectSoundSource& DirectSoundSource : DirectSoundSources)
 		{
 			DirectSoundSource.InBuffer.format = InputAudioFormat;
-			DirectSoundSource.InBuffer.numSamples = FrameSize;
+			DirectSoundSource.InBuffer.numSamples = InitializationParams.BufferLength;
 			DirectSoundSource.InBuffer.interleavedBuffer = nullptr;
 			DirectSoundSource.InBuffer.deinterleavedBuffer = nullptr;
 
 			DirectSoundSource.OutBuffer.format = OutputAudioFormat;
-			DirectSoundSource.OutBuffer.numSamples = FrameSize;
+			DirectSoundSource.OutBuffer.numSamples = InitializationParams.BufferLength;
 			DirectSoundSource.OutBuffer.interleavedBuffer = nullptr;
 			DirectSoundSource.OutBuffer.deinterleavedBuffer = nullptr;
 		}
-
-		SteamAudioModule = &FModuleManager::GetModuleChecked<FSteamAudioModule>("SteamAudio");
 	}
 
 	void FPhononOcclusion::OnInitSource(const uint32 SourceId, const FName& AudioComponentUserId, const uint32 NumChannels, UOcclusionPluginSourceSettingsBase* InSettings)
@@ -103,7 +101,7 @@ namespace SteamAudio
 
 	void FPhononOcclusion::ProcessAudio(const FAudioPluginSourceInputData& InputData, FAudioPluginSourceOutputData& OutputData)
 	{
-		auto& DirectSoundSource = DirectSoundSources[InputData.SourceId];
+		FDirectSoundSource& DirectSoundSource = DirectSoundSources[InputData.SourceId];
 
 		if (!EnvironmentalRenderer)
 		{
@@ -125,18 +123,17 @@ namespace SteamAudio
 		DirectSoundEffectOptions.applyDistanceAttenuation = static_cast<IPLbool>(DirectSoundSources[InputData.SourceId].bDirectAttenuation);
 		DirectSoundEffectOptions.directOcclusionMode = static_cast<IPLDirectOcclusionMode>(DirectSoundSources[InputData.SourceId].DirectOcclusionMode);
 
-		iplApplyDirectSoundEffect(DirectSoundSource.DirectSoundEffect, DirectSoundSource.InBuffer, DirectSoundSource.DirectSoundPath,
-			DirectSoundEffectOptions, DirectSoundSource.OutBuffer);
+		iplApplyDirectSoundEffect(DirectSoundSource.DirectSoundEffect, DirectSoundSource.InBuffer, DirectSoundSource.DirectSoundPath, DirectSoundEffectOptions, DirectSoundSource.OutBuffer);
 	}
 
 	void FPhononOcclusion::UpdateDirectSoundSources(const FVector& ListenerPosition, const FVector& ListenerForward, const FVector& ListenerUp)
 	{
-		if (!EnvironmentalRenderer)
+		if (!EnvironmentalRenderer || !EnvironmentCriticalSectionHandle)
 		{
 			return;
 		}
 
-		FScopeLock EnvironmentLock(&SteamAudioModule->GetEnvironmentCriticalSection());
+		FScopeLock EnvironmentLock(EnvironmentCriticalSectionHandle);
 
 		for (FDirectSoundSource& DirectSoundSource : DirectSoundSources)
 		{
@@ -160,6 +157,12 @@ namespace SteamAudio
 	{
 		EnvironmentalRenderer = InEnvironmentalRenderer;
 	}
+
+	void FPhononOcclusion::SetCriticalSectionHandle(FCriticalSection* CriticalSectionHandle)
+	{
+		EnvironmentCriticalSectionHandle = CriticalSectionHandle;
+	}
+
 }
 
 //==================================================================================================================================================
diff --git a/Engine/Plugins/Runtime/Steam/SteamAudio/Source/SteamAudio/Private/PhononOcclusion.h b/Engine/Plugins/Runtime/Steam/SteamAudio/Source/SteamAudio/Private/PhononOcclusion.h
--- a/Engine/Plugins/Runtime/Steam/SteamAudio/Source/SteamAudio/Private/PhononOcclusion.h
+++ b/Engine/Plugins/Runtime/Steam/SteamAudio/Source/SteamAudio/Private/PhononOcclusion.h
@@ -7,7 +7,6 @@
 #include "IAudioExtensionPlugin.h"
 #include "PhononCommon.h"
 #include "phonon.h"
-#include "SteamAudioModule.h"
 
 class UOcclusionPluginSourceSettingsBase;
 
@@ -31,25 +30,47 @@ namespace SteamAudio
 		bool bNeedsUpdate;
 	};
 
+	/************************************************************************/
+	/* FPhononOcclusion                                                     */
+	/* Scene-dependent audio occlusion plugin. Receives updates from        */
+	/* an FPhononPluginManager on the game thread on player position and    */
+	/* geometry, and performs geometry-aware filtering of the direct path   */
+	/* of an audio source.                                                  */
+	/************************************************************************/
 	class FPhononOcclusion : public IAudioOcclusion
 	{
 	public:
 		FPhononOcclusion();
 		~FPhononOcclusion();
 
-		virtual void Initialize(const int32 SampleRate, const int32 NumSources, const int32 FrameSize) override;
+		//~ Begin IAudioOcclusion
+		virtual void Initialize(const FAudioPluginInitializationParams InitializationParams) override;
 		virtual void OnInitSource(const uint32 SourceId, const FName& AudioComponentUserId, const uint32 NumChannels, UOcclusionPluginSourceSettingsBase* InSettings) override;
 		virtual void OnReleaseSource(const uint32 SourceId) override;
 		virtual void ProcessAudio(const FAudioPluginSourceInputData& InputData, FAudioPluginSourceOutputData& OutputData) override;
+		//~ End IAudioOcclusion
 
+		// Receives updates on listener positions from the game thread using this function call.
 		void UpdateDirectSoundSources(const FVector& ListenerPosition, const FVector& ListenerForward, const FVector& ListenerUp);
+
+		// Sets up handle to the environmental handle owned by the FPhononPluginManager.
 		void SetEnvironmentalRenderer(IPLhandle EnvironmentalRenderer);
 
+		// Sets up the handle to the Critical Section owned by the FPhononPluginManager.
+		void SetCriticalSectionHandle(FCriticalSection* CriticalSectionHandle);
+
 	private:
+		//Critical section. Scope locked so that environment is not modified by the Plugin Manager during audio processing.
+		FCriticalSection* EnvironmentCriticalSectionHandle;
+		
+		//Handle to FPhononPluginManager's Environmental Renderer.
 		IPLhandle EnvironmentalRenderer;
+
+		//Cached Audio Formats:
 		IPLAudioFormat InputAudioFormat;
 		IPLAudioFormat OutputAudioFormat;
+
+		//Cached array of direct sound sources to be occluded.
 		TArray<FDirectSoundSource> DirectSoundSources;
-		FSteamAudioModule* SteamAudioModule;
 	};
 }
