Remove Android support for the Oculus audio plugin.

diff --git a/Engine/Plugins/Runtime/Oculus/OculusAudio/OculusAudio.uplugin b/Engine/Plugins/Runtime/Oculus/OculusAudio/OculusAudio.uplugin
--- a/Engine/Plugins/Runtime/Oculus/OculusAudio/OculusAudio.uplugin
+++ b/Engine/Plugins/Runtime/Oculus/OculusAudio/OculusAudio.uplugin
@@ -21,7 +21,6 @@
 			"LoadingPhase": "PostConfigInit",
 			"WhitelistPlatforms": [
 				"Win64",
-				"Android"
 			]
 		},
 		{
