Fixing oculus crash when first enabling it in an empty project

diff --git a/Engine/Plugins/Runtime/Oculus/OculusAudio/Source/OculusAudio/Private/OculusAmbisonicSpatializer.cpp b/Engine/Plugins/Runtime/Oculus/OculusAudio/Source/OculusAudio/Private/OculusAmbisonicSpatializer.cpp
--- a/Engine/Plugins/Runtime/Oculus/OculusAudio/Source/OculusAudio/Private/OculusAmbisonicSpatializer.cpp
+++ b/Engine/Plugins/Runtime/Oculus/OculusAudio/Source/OculusAudio/Private/OculusAmbisonicSpatializer.cpp
@@ -57,21 +57,23 @@ void FOculusAmbisonicsMixer::EncodeToAmbisonics(const uint32 SourceId, const FAm
 
 void FOculusAmbisonicsMixer::OnOpenDecodingStream(const uint32 StreamId, UAmbisonicsSubmixSettingsBase* InSettings, FAmbisonicsDecoderPositionalData& SpecifiedOutputPositions)
 {
-	UOculusAmbisonicsSettings* Settings = CastChecked<UOculusAmbisonicsSettings>(InSettings);
+	UOculusAmbisonicsSettings* Settings = Cast<UOculusAmbisonicsSettings>(InSettings);
+	if (Settings)
+	{
+		ovrAudioAmbisonicStream NewStream = nullptr;
+		ovrAudioAmbisonicFormat StreamFormat = (Settings->ChannelOrder == EAmbisonicFormat::AmbiX ? ovrAudioAmbisonicFormat_AmbiX : ovrAudioAmbisonicFormat_FuMa);
+		ovrResult OurResult = ovrAudio_CreateAmbisonicStream(Context, SampleRate, BufferLength, StreamFormat, 1, &NewStream);
 
-	ovrAudioAmbisonicStream NewStream = nullptr;
-	ovrAudioAmbisonicFormat StreamFormat = (Settings->ChannelOrder == EAmbisonicFormat::AmbiX ? ovrAudioAmbisonicFormat_AmbiX : ovrAudioAmbisonicFormat_FuMa);
-	ovrResult OurResult = ovrAudio_CreateAmbisonicStream(Context, SampleRate, BufferLength, StreamFormat, 1, &NewStream);
+		check(OurResult == 0);
+		check(NewStream != nullptr);
 
-	check(OurResult == 0);
-	check(NewStream != nullptr);
-
-	ovrAudioAmbisonicSpeakerLayout SpeakerLayout = (Settings->SpatializationMode == EAmbisonicMode::SphericalHarmonics) ?
-		ovrAudioAmbisonicSpeakerLayout_SphericalHarmonics : ovrAudioAmbisonicSpeakerLayout_Icosahedron;
-	OurResult = ovrAudio_SetAmbisonicSpeakerLayout(NewStream, SpeakerLayout);
+		ovrAudioAmbisonicSpeakerLayout SpeakerLayout = (Settings->SpatializationMode == EAmbisonicMode::SphericalHarmonics) ?
+			ovrAudioAmbisonicSpeakerLayout_SphericalHarmonics : ovrAudioAmbisonicSpeakerLayout_Icosahedron;
+		OurResult = ovrAudio_SetAmbisonicSpeakerLayout(NewStream, SpeakerLayout);
 
-	check(OurResult == 0);
-	OpenStreams.Add(StreamId, NewStream);
+		check(OurResult == 0);
+		OpenStreams.Add(StreamId, NewStream);
+	}
 }
 
 void FOculusAmbisonicsMixer::OnCloseDecodingStream(const uint32 StreamId)
diff --git a/Engine/Source/Runtime/AudioMixer/Private/AudioMixerSubmix.cpp b/Engine/Source/Runtime/AudioMixer/Private/AudioMixerSubmix.cpp
--- a/Engine/Source/Runtime/AudioMixer/Private/AudioMixerSubmix.cpp
+++ b/Engine/Source/Runtime/AudioMixer/Private/AudioMixerSubmix.cpp
@@ -93,6 +93,11 @@ namespace Audio
 				//If we do have a valid ambisonics decoder, lets use it. Otherwise, treat this submix like a device submix.
 				if (AmbisonicsMixer.IsValid())
 				{
+					if (!InSoundSubmix->AmbisonicsPluginSettings)
+					{
+						InSoundSubmix->AmbisonicsPluginSettings = AmbisonicsMixer->GetDefaultSettings();
+					}
+
 					if (InSoundSubmix->AmbisonicsPluginSettings != nullptr)
 					{
 						OnAmbisonicsSettingsChanged(InSoundSubmix->AmbisonicsPluginSettings);
