[NO NON-PII COMMIT MESSAGE]

diff --git a/Engine/Plugins/Runtime/Oculus/OculusAudio/Source/OculusAudio/Private/OculusAmbisonicSpatializer.cpp b/Engine/Plugins/Runtime/Oculus/OculusAudio/Source/OculusAudio/Private/OculusAmbisonicSpatializer.cpp
--- a/Engine/Plugins/Runtime/Oculus/OculusAudio/Source/OculusAudio/Private/OculusAmbisonicSpatializer.cpp
+++ b/Engine/Plugins/Runtime/Oculus/OculusAudio/Source/OculusAudio/Private/OculusAmbisonicSpatializer.cpp
@@ -12,7 +12,6 @@ FOculusAmbisonicsMixer::FOculusAmbisonicsMixer()
 
 void FOculusAmbisonicsMixer::Shutdown()
 {
-	OVRA_CALL(ovrAudio_DestroyContext)(Context);
 }
 
 int32 FOculusAmbisonicsMixer::GetNumChannelsForAmbisonicsFormat(UAmbisonicsSubmixSettingsBase* InSettings)
diff --git a/Engine/Plugins/Runtime/Oculus/OculusAudio/Source/OculusAudio/Private/OculusAudioContextManager.cpp b/Engine/Plugins/Runtime/Oculus/OculusAudio/Source/OculusAudio/Private/OculusAudioContextManager.cpp
--- a/Engine/Plugins/Runtime/Oculus/OculusAudio/Source/OculusAudio/Private/OculusAudioContextManager.cpp
+++ b/Engine/Plugins/Runtime/Oculus/OculusAudio/Source/OculusAudio/Private/OculusAudioContextManager.cpp
@@ -75,7 +75,25 @@ void FOculusAudioContextManager::OnListenerInitialize(FAudioDevice* AudioDevice,
 
 void FOculusAudioContextManager::OnListenerShutdown(FAudioDevice* AudioDevice)
 {
-	
+	FOculusAudioPlugin* Plugin = &FModuleManager::GetModuleChecked<FOculusAudioPlugin>("OculusAudio");
+	check(Plugin != nullptr);
+
+	FString OculusSpatializerPluginName = Plugin->GetSpatializationPluginFactory()->GetDisplayName();
+	FString CurrentSpatializerPluginName = AudioPluginUtilities::GetDesiredPluginName(EAudioPlugin::SPATIALIZATION, AudioPluginUtilities::CurrentPlatform);
+	if (CurrentSpatializerPluginName.Equals(OculusSpatializerPluginName))
+	{
+		OculusAudioSpatializationAudioMixer* Spatializer =
+			static_cast<OculusAudioSpatializationAudioMixer*>(AudioDevice->SpatializationPluginInterface.Get());
+		Spatializer->SetContext(nullptr);
+	}
+
+	FString OculusReverbPluginName = Plugin->GetReverbPluginFactory()->GetDisplayName();
+	FString CurrentReverbPluginName = AudioPluginUtilities::GetDesiredPluginName(EAudioPlugin::REVERB, AudioPluginUtilities::CurrentPlatform);
+	if (CurrentReverbPluginName.Equals(OculusReverbPluginName))
+	{
+		OculusAudioReverb* Reverb = static_cast<OculusAudioReverb*>(AudioDevice->ReverbPluginInterface.Get());
+		Reverb->SetContext(nullptr);
+	}
 }
 
 ovrAudioContext FOculusAudioContextManager::GetOrCreateSerializationContext(UActorComponent* Parent)
diff --git a/Engine/Plugins/Runtime/Oculus/OculusAudio/Source/OculusAudio/Private/OculusAudioMixer.cpp b/Engine/Plugins/Runtime/Oculus/OculusAudio/Source/OculusAudio/Private/OculusAudioMixer.cpp
--- a/Engine/Plugins/Runtime/Oculus/OculusAudio/Source/OculusAudio/Private/OculusAudioMixer.cpp
+++ b/Engine/Plugins/Runtime/Oculus/OculusAudio/Source/OculusAudio/Private/OculusAudioMixer.cpp
@@ -30,9 +30,14 @@ OculusAudioSpatializationAudioMixer::~OculusAudioSpatializationAudioMixer()
 
 void OculusAudioSpatializationAudioMixer::SetContext(ovrAudioContext* SharedContext)
 {
+	FScopeLock ScopeLock(&ContextLock);
+
 	Context = SharedContext;
 
-	check(Context != nullptr);
+	if (SharedContext == nullptr)
+	{
+		return;
+	}
 
 	Params.AddDefaulted(InitParams.NumSources);
 
@@ -82,6 +87,7 @@ void OculusAudioSpatializationAudioMixer::Initialize(const FAudioPluginInitializ
 
 void OculusAudioSpatializationAudioMixer::ApplyOculusAudioSettings(const UOculusAudioSettings* Settings)
 {
+	FScopeLock ScopeLock(&ContextLock);
 	ovrResult Result = OVRA_CALL(ovrAudio_Enable)(*Context, ovrAudioEnable_SimpleRoomModeling, Settings->EarlyReflections);
 	OVR_AUDIO_CHECK(Result, "Failed to enable reflections");
 
@@ -119,6 +125,7 @@ bool OculusAudioSpatializationAudioMixer::IsSpatializationEffectInitialized() co
 
 void OculusAudioSpatializationAudioMixer::OnInitSource(const uint32 SourceId, const FName& AudioComponentUserId, USpatializationPluginSourceSettingsBase* InSettings)
 {
+	FScopeLock ScopeLock(&ContextLock);
 	check(bOvrContextInitialized);
 
 	if (InSettings != nullptr) 
@@ -154,6 +161,7 @@ void OculusAudioSpatializationAudioMixer::SetSpatializationParameters(uint32 Voi
 
 void OculusAudioSpatializationAudioMixer::ProcessAudio(const FAudioPluginSourceInputData& InputData, FAudioPluginSourceOutputData& OutputData)
 {
+	FScopeLock ScopeLock(&ContextLock);
 	if (InputData.SpatializationParams && *Context)
 	{
 		Params[InputData.SourceId] = *InputData.SpatializationParams;
@@ -185,14 +193,22 @@ void OculusAudioSpatializationAudioMixer::ProcessAudio(const FAudioPluginSourceI
 
 bool OculusAudioSpatializationAudioMixer::Tick(float DeltaTime)
 {
-	if (Context != nullptr)
+	if (ContextLock.TryLock())
 	{
-		ovrResult Result = OVRA_CALL(ovrAudio_UpdateRoomModel)(*Context, 1.0f);
-
-		UOculusAudioSettings* settings = GetMutableDefault<UOculusAudioSettings>();
-		Result = OVRA_CALL(ovrAudio_SetPropagationQuality)(*Context, settings->PropagationQuality);
-		check(Result == ovrSuccess);
+		if (Context != nullptr && *Context == nullptr)
+		{
+			ovrResult Result = OVRA_CALL(ovrAudio_UpdateRoomModel)(*Context, 1.0f);
+
+			UOculusAudioSettings* settings = GetMutableDefault<UOculusAudioSettings>();
+			Result = OVRA_CALL(ovrAudio_SetPropagationQuality)(*Context, settings->PropagationQuality);
+			if (Result != ovrSuccess)
+			{
+				UE_LOG(LogTemp, Warning, TEXT("Bad Propagation Quality setting %d!"), settings->PropagationQuality);
+			}
+		}
+
+		ContextLock.Unlock();
 	}
-
+	
 	return true;
 }
\ No newline at end of file
diff --git a/Engine/Plugins/Runtime/Oculus/OculusAudio/Source/OculusAudio/Private/OculusAudioMixer.h b/Engine/Plugins/Runtime/Oculus/OculusAudio/Source/OculusAudio/Private/OculusAudioMixer.h
--- a/Engine/Plugins/Runtime/Oculus/OculusAudio/Source/OculusAudio/Private/OculusAudioMixer.h
+++ b/Engine/Plugins/Runtime/Oculus/OculusAudio/Source/OculusAudio/Private/OculusAudioMixer.h
@@ -50,6 +50,8 @@ private:
 	TArray<FSpatializationParams> Params;
 
 	ovrAudioContext* Context;
+	FCriticalSection ContextLock;
+
 	FAudioPluginInitializationParams InitParams;
 	FDelegateHandle TickDelegateHandle;
 };
\ No newline at end of file
diff --git a/Engine/Plugins/Runtime/Oculus/OculusAudio/Source/OculusAudio/Private/OculusAudioReverb.cpp b/Engine/Plugins/Runtime/Oculus/OculusAudio/Source/OculusAudio/Private/OculusAudioReverb.cpp
--- a/Engine/Plugins/Runtime/Oculus/OculusAudio/Source/OculusAudio/Private/OculusAudioReverb.cpp
+++ b/Engine/Plugins/Runtime/Oculus/OculusAudio/Source/OculusAudio/Private/OculusAudioReverb.cpp
@@ -6,8 +6,8 @@
 
 void FSubmixEffectOculusReverbPlugin::SetContext(ovrAudioContext* SharedContext)
 {
+	FScopeLock ScopeLock(&ContextLock);
 	Context = SharedContext;
-	check(Context != nullptr);
 }
 
 FSubmixEffectOculusReverbPlugin::FSubmixEffectOculusReverbPlugin()
@@ -17,6 +17,7 @@ FSubmixEffectOculusReverbPlugin::FSubmixEffectOculusReverbPlugin()
 
 void FSubmixEffectOculusReverbPlugin::OnProcessAudio(const FSoundEffectSubmixInputData& InputData, FSoundEffectSubmixOutputData& OutputData)
 {
+	FScopeLock ScopeLock(&ContextLock);
 	if (Context != nullptr && *Context != nullptr)
 	{
 		int Enabled = 0;
@@ -34,8 +35,11 @@ void FSubmixEffectOculusReverbPlugin::OnProcessAudio(const FSoundEffectSubmixInp
 
 void OculusAudioReverb::SetContext(ovrAudioContext* SharedContext)
 {
-	check(SharedContext != nullptr);
-	Context = SharedContext;
+	if (SharedContext != nullptr)
+	{
+		Context = SharedContext;
+	}
+	
 	for (FSubmixEffectOculusReverbPlugin* Submix : Submixes)
 	{
 		Submix->SetContext(SharedContext);
diff --git a/Engine/Plugins/Runtime/Oculus/OculusAudio/Source/OculusAudio/Private/OculusAudioReverb.h b/Engine/Plugins/Runtime/Oculus/OculusAudio/Source/OculusAudio/Private/OculusAudioReverb.h
--- a/Engine/Plugins/Runtime/Oculus/OculusAudio/Source/OculusAudio/Private/OculusAudioReverb.h
+++ b/Engine/Plugins/Runtime/Oculus/OculusAudio/Source/OculusAudio/Private/OculusAudioReverb.h
@@ -30,6 +30,7 @@ public:
 	}
 private:
 	ovrAudioContext* Context;
+	FCriticalSection ContextLock;
 };
 
 /************************************************************************/
diff --git a/Engine/Source/Runtime/AudioMixer/Private/AudioMixerDevice.cpp b/Engine/Source/Runtime/AudioMixer/Private/AudioMixerDevice.cpp
--- a/Engine/Source/Runtime/AudioMixer/Private/AudioMixerDevice.cpp
+++ b/Engine/Source/Runtime/AudioMixer/Private/AudioMixerDevice.cpp
@@ -70,6 +70,18 @@ namespace Audio
 #endif
 	}
 
+	void FMixerDevice::OnListenerUpdated(const TArray<FListener>& InListeners)
+	{
+		ListenerTransforms.Reset(InListeners.Num());
+
+		for (const FListener& Listener : InListeners)
+		{
+			ListenerTransforms.Add(Listener.Transform);
+		}
+
+		SourceManager.SetListenerTransforms(ListenerTransforms);
+	}
+
 	void FMixerDevice::ResetAudioRenderingThreadId()
 	{
 #if AUDIO_MIXER_ENABLE_DEBUG_MODE
diff --git a/Engine/Source/Runtime/AudioMixer/Public/AudioMixerDevice.h b/Engine/Source/Runtime/AudioMixer/Public/AudioMixerDevice.h
--- a/Engine/Source/Runtime/AudioMixer/Public/AudioMixerDevice.h
+++ b/Engine/Source/Runtime/AudioMixer/Public/AudioMixerDevice.h
@@ -201,6 +201,12 @@ namespace Audio
 		double GetAudioRenderThreadTime() const { return AudioThreadTimingData.AudioRenderThreadTime; }
 		double GetAudioClockDelta() const { return AudioClockDelta; }
 
+	protected:
+
+		virtual void OnListenerUpdated(const TArray<FListener>& InListeners) override;
+
+		TArray<FTransform> ListenerTransforms;
+
 	private:
 		// Resets the thread ID used for audio rendering
 		void ResetAudioRenderingThreadId();
diff --git a/Engine/Source/Runtime/Engine/Private/AudioDevice.cpp b/Engine/Source/Runtime/Engine/Private/AudioDevice.cpp
--- a/Engine/Source/Runtime/Engine/Private/AudioDevice.cpp
+++ b/Engine/Source/Runtime/Engine/Private/AudioDevice.cpp
@@ -2886,6 +2886,8 @@ void FAudioDevice::SetListener(UWorld* World, const int32 InViewportIndex, const
 		Listener.WorldID = WorldID;
 		Listener.Transform = ListenerTransformCopy;
 
+		OnListenerUpdated(AudioThreadListeners);
+
 	}, GET_STATID(STAT_AudioSetListener));
 }
 
diff --git a/Engine/Source/Runtime/Engine/Public/AudioDevice.h b/Engine/Source/Runtime/Engine/Public/AudioDevice.h
--- a/Engine/Source/Runtime/Engine/Public/AudioDevice.h
+++ b/Engine/Source/Runtime/Engine/Public/AudioDevice.h
@@ -1208,6 +1208,11 @@ protected:
 	 */
 	void StartSources(TArray<FWaveInstance*>& WaveInstances, int32 FirstActiveIndex, bool bGameTicking);
 
+	/**
+	 * This is overridden in Audio::FMixerDevice to propogate listener information to the audio thread.
+	 */
+	virtual void OnListenerUpdated(const TArray<FListener>& InListeners) {};
+
 private:
 
 	/**
