Fix for Oculus spatialization plugin crash on PIE shutdown (Tick() was being called by delegate on GC'd object)

diff --git a/Engine/Plugins/Runtime/Oculus/OculusAudio/Source/OculusAudio/Private/OculusAudioMixer.cpp b/Engine/Plugins/Runtime/Oculus/OculusAudio/Source/OculusAudio/Private/OculusAudioMixer.cpp
--- a/Engine/Plugins/Runtime/Oculus/OculusAudio/Source/OculusAudio/Private/OculusAudioMixer.cpp
+++ b/Engine/Plugins/Runtime/Oculus/OculusAudio/Source/OculusAudio/Private/OculusAudioMixer.cpp
@@ -19,10 +19,13 @@ OculusAudioSpatializationAudioMixer::OculusAudioSpatializationAudioMixer()
 
 OculusAudioSpatializationAudioMixer::~OculusAudioSpatializationAudioMixer()
 {
-	if (bOvrContextInitialized)
+	if (TickDelegateHandle.IsValid())
 	{
 		FTicker::GetCoreTicker().RemoveTicker(TickDelegateHandle);
+	}
 
+	if (bOvrContextInitialized)
+	{
 		// clear context from map
 		Shutdown();
 	}
