Fixing serialization overread on OculusAudioGeometryLandscape and OculusAudioGeometryBSP

diff --git a/Engine/Plugins/Runtime/Oculus/OculusAudio/Source/OculusAudio/Private/OculusAudioGeometryBSP.cpp b/Engine/Plugins/Runtime/Oculus/OculusAudio/Source/OculusAudio/Private/OculusAudioGeometryBSP.cpp
--- a/Engine/Plugins/Runtime/Oculus/OculusAudio/Source/OculusAudio/Private/OculusAudioGeometryBSP.cpp
+++ b/Engine/Plugins/Runtime/Oculus/OculusAudio/Source/OculusAudio/Private/OculusAudioGeometryBSP.cpp
@@ -24,6 +24,15 @@ void UOculusAudioGeometryBSP::Serialize(FArchive & Ar)
 		static size_t Read(void* userData, void* bytes, size_t byteCount) {
 			FArchive* userAr = static_cast<FArchive*>(userData);
 			check(userAr->IsLoading());
+
+			// Check that we are not reading past end of archive..
+			int64 CurrentPosition = userAr->Tell();
+			int64 TotalSize = userAr->TotalSize();
+			if ((CurrentPosition + static_cast<int64>(byteCount)) > TotalSize)
+			{
+				return 0;
+			}
+
 			userAr->Serialize(bytes, byteCount);
 			return userAr->GetError() ? 0 : byteCount;
 		}
diff --git a/Engine/Plugins/Runtime/Oculus/OculusAudio/Source/OculusAudio/Private/OculusAudioGeometryLandscape.cpp b/Engine/Plugins/Runtime/Oculus/OculusAudio/Source/OculusAudio/Private/OculusAudioGeometryLandscape.cpp
--- a/Engine/Plugins/Runtime/Oculus/OculusAudio/Source/OculusAudio/Private/OculusAudioGeometryLandscape.cpp
+++ b/Engine/Plugins/Runtime/Oculus/OculusAudio/Source/OculusAudio/Private/OculusAudioGeometryLandscape.cpp
@@ -26,6 +26,15 @@ void UOculusAudioGeometryLandscape::Serialize(FArchive & Ar)
 		static size_t Read(void* userData, void* bytes, size_t byteCount) {
 			FArchive* userAr = static_cast<FArchive*>(userData);
 			check(userAr->IsLoading());
+
+			// Check that we are not reading past end of archive..
+			int64 CurrentPosition = userAr->Tell();
+			int64 TotalSize = userAr->TotalSize();
+			if ((CurrentPosition + static_cast<int64>(byteCount)) > TotalSize)
+			{
+				return 0;
+			}
+
 			userAr->Serialize(bytes, byteCount);
 			return userAr->GetError() ? 0 : byteCount;
 		}
