Fix for bad ptr in log message

diff --git a/Engine/Plugins/Runtime/Oculus/OculusAudio/Source/OculusAudio/Private/OculusAudioDllManager.cpp b/Engine/Plugins/Runtime/Oculus/OculusAudio/Source/OculusAudio/Private/OculusAudioDllManager.cpp
--- a/Engine/Plugins/Runtime/Oculus/OculusAudio/Source/OculusAudio/Private/OculusAudioDllManager.cpp
+++ b/Engine/Plugins/Runtime/Oculus/OculusAudio/Source/OculusAudio/Private/OculusAudioDllManager.cpp
@@ -128,7 +128,7 @@ bool FOculusAudioLibraryManager::LoadDll()
 		}
 
 
-		UE_LOG(LogAudio, Display, TEXT("Attempting to load Oculus Spatializer DLL: %s (from %s)"), *Path, *DLL_NAME);
+		UE_LOG(LogAudio, Display, TEXT("Attempting to load Oculus Spatializer DLL: %s (from %s)"), *Path, DLL_NAME);
 
 		FPlatformProcess::PushDllDirectory(*Path);
 		OculusAudioDllHandle = FPlatformProcess::GetDllHandle(*(Path + DLL_NAME)); // FIXME: support win32
