I think this is what the build error from clang wants us to change.  Don't have clang installed to test, so will just let the build machine decide.

diff --git a/Engine/Plugins/Runtime/Oculus/OculusAudio/Source/OculusAudio/Private/OculusAudioGeometryLandscape.cpp b/Engine/Plugins/Runtime/Oculus/OculusAudio/Source/OculusAudio/Private/OculusAudioGeometryLandscape.cpp
--- a/Engine/Plugins/Runtime/Oculus/OculusAudio/Source/OculusAudio/Private/OculusAudioGeometryLandscape.cpp
+++ b/Engine/Plugins/Runtime/Oculus/OculusAudio/Source/OculusAudio/Private/OculusAudioGeometryLandscape.cpp
@@ -133,7 +133,7 @@ void UOculusAudioGeometryLandscape::CreateGeometryFromLandscape(ovrAudioContext
 			}
 		}
 
-		ovrAudioMesh ovrMesh = { 0 };
+		ovrAudioMesh ovrMesh = { { 0 } };
 
 		ovrAudioMeshVertices ovrVertices = { 0 };
 		ovrVertices.vertices = Vertices.GetData();
