Initialize some variables.

diff --git a/Engine/Plugins/Runtime/Oculus/OculusAudio/Source/OculusAudio/Private/OculusAudioDllManager.cpp b/Engine/Plugins/Runtime/Oculus/OculusAudio/Source/OculusAudio/Private/OculusAudioDllManager.cpp
--- a/Engine/Plugins/Runtime/Oculus/OculusAudio/Source/OculusAudio/Private/OculusAudioDllManager.cpp
+++ b/Engine/Plugins/Runtime/Oculus/OculusAudio/Source/OculusAudio/Private/OculusAudioDllManager.cpp
@@ -50,9 +50,9 @@ void FOculusAudioLibraryManager::Initialize()
 	if (!bInitialized)
 	{
 		// Check the version number
-		int32 MajorVersionNumber;
-		int32 MinorVersionNumber;
-		int32 PatchNumber;
+		int32 MajorVersionNumber = -1;
+		int32 MinorVersionNumber = -1;
+		int32 PatchNumber = -1;
 
 		const char* OvrVersionString = OVRA_CALL(ovrAudio_GetVersion)(&MajorVersionNumber, &MinorVersionNumber, &PatchNumber);
 		if (MajorVersionNumber != OVR_AUDIO_MAJOR_VERSION || MinorVersionNumber != OVR_AUDIO_MINOR_VERSION)
