Removing temp optimization pragma

diff --git a/Engine/Plugins/Runtime/Oculus/OculusAudio/Source/OculusAudio/Private/OculusAudioMaterialComponent.cpp b/Engine/Plugins/Runtime/Oculus/OculusAudio/Source/OculusAudio/Private/OculusAudioMaterialComponent.cpp
--- a/Engine/Plugins/Runtime/Oculus/OculusAudio/Source/OculusAudio/Private/OculusAudioMaterialComponent.cpp
+++ b/Engine/Plugins/Runtime/Oculus/OculusAudio/Source/OculusAudio/Private/OculusAudioMaterialComponent.cpp
@@ -354,6 +354,4 @@ void UOculusAudioMaterialComponent::ResetAcousticMaterialPreset()
 		Scattering8000Hz = Scattering[7];
 	}
 
-#endif // WITH_EDITOR
-
-#pragma optimize("", on) //PAS - TEMP - DELETE THIS - better debugging
\ No newline at end of file
+#endif // WITH_EDITOR
\ No newline at end of file
