[NO NON-PII COMMIT MESSAGE]

diff --git a/Engine/Plugins/Online/OnlineSubsystemOculus/Source/Private/OculusIdentityCallbackProxy.cpp b/Engine/Plugins/Online/OnlineSubsystemOculus/Source/Private/OculusIdentityCallbackProxy.cpp
--- a/Engine/Plugins/Online/OnlineSubsystemOculus/Source/Private/OculusIdentityCallbackProxy.cpp
+++ b/Engine/Plugins/Online/OnlineSubsystemOculus/Source/Private/OculusIdentityCallbackProxy.cpp
@@ -23,7 +23,7 @@ void UOculusIdentityCallbackProxy::Activate()
 
 	if (OculusIdentityInterface.IsValid())
 	{
-		DelegateHandle = Online::GetIdentityInterface()->AddOnLoginCompleteDelegate_Handle(
+		DelegateHandle = OculusIdentityInterface->AddOnLoginCompleteDelegate_Handle(
 			0, 
 			FOnLoginCompleteDelegate::CreateUObject(this, &UOculusIdentityCallbackProxy::OnLoginCompleteDelegate)
 		);
diff --git a/Engine/Plugins/Runtime/Oculus/OculusAudio/Source/OculusAudio/Private/OculusAmbisonicSpatializer.cpp b/Engine/Plugins/Runtime/Oculus/OculusAudio/Source/OculusAudio/Private/OculusAmbisonicSpatializer.cpp
--- a/Engine/Plugins/Runtime/Oculus/OculusAudio/Source/OculusAudio/Private/OculusAmbisonicSpatializer.cpp
+++ b/Engine/Plugins/Runtime/Oculus/OculusAudio/Source/OculusAudio/Private/OculusAmbisonicSpatializer.cpp
@@ -57,23 +57,26 @@ void FOculusAmbisonicsMixer::EncodeToAmbisonics(const uint32 SourceId, const FAm
 
 void FOculusAmbisonicsMixer::OnOpenDecodingStream(const uint32 StreamId, UAmbisonicsSubmixSettingsBase* InSettings, FAmbisonicsDecoderPositionalData& SpecifiedOutputPositions)
 {
-	UOculusAmbisonicsSettings* Settings = Cast<UOculusAmbisonicsSettings>(InSettings);
-	if (Settings)
+	ovrAudioAmbisonicFormat StreamFormat = ovrAudioAmbisonicFormat_AmbiX;
+	ovrAudioAmbisonicSpeakerLayout SpeakerLayout = ovrAudioAmbisonicSpeakerLayout_SphericalHarmonics;
+	if (InSettings != nullptr) 
 	{
-		ovrAudioAmbisonicStream NewStream = nullptr;
-		ovrAudioAmbisonicFormat StreamFormat = (Settings->ChannelOrder == EAmbisonicFormat::AmbiX ? ovrAudioAmbisonicFormat_AmbiX : ovrAudioAmbisonicFormat_FuMa);
-		ovrResult OurResult = OVRA_CALL(ovrAudio_CreateAmbisonicStream)(Context, SampleRate, BufferLength, StreamFormat, 1, &NewStream);
+		UOculusAmbisonicsSettings* Settings = CastChecked<UOculusAmbisonicsSettings>(InSettings);
+		StreamFormat = (Settings->ChannelOrder == EAmbisonicFormat::AmbiX ? ovrAudioAmbisonicFormat_AmbiX : ovrAudioAmbisonicFormat_FuMa);
+		SpeakerLayout = (Settings->SpatializationMode == EAmbisonicMode::SphericalHarmonics) ?
+			ovrAudioAmbisonicSpeakerLayout_SphericalHarmonics : ovrAudioAmbisonicSpeakerLayout_Icosahedron;
+	}
 
-		check(OurResult == 0);
-		check(NewStream != nullptr);
+	ovrAudioAmbisonicStream NewStream = nullptr;
+	ovrResult OurResult = OVRA_CALL(ovrAudio_CreateAmbisonicStream)(Context, SampleRate, BufferLength, StreamFormat, 1, &NewStream);
 
-		ovrAudioAmbisonicSpeakerLayout SpeakerLayout = (Settings->SpatializationMode == EAmbisonicMode::SphericalHarmonics) ?
-			ovrAudioAmbisonicSpeakerLayout_SphericalHarmonics : ovrAudioAmbisonicSpeakerLayout_Icosahedron;
-		OurResult = OVRA_CALL(ovrAudio_SetAmbisonicSpeakerLayout)(NewStream, SpeakerLayout);
+	check(OurResult == 0);
+	check(NewStream != nullptr);
 
-		check(OurResult == 0);
-		OpenStreams.Add(StreamId, NewStream);
-	}
+	OurResult = OVRA_CALL(ovrAudio_SetAmbisonicSpeakerLayout)(NewStream, SpeakerLayout);
+
+	check(OurResult == 0);
+	OpenStreams.Add(StreamId, NewStream);
 }
 
 void FOculusAmbisonicsMixer::OnCloseDecodingStream(const uint32 StreamId)
diff --git a/Engine/Plugins/Runtime/Oculus/OculusAudio/Source/OculusAudio/Private/OculusAudioDllManager.h b/Engine/Plugins/Runtime/Oculus/OculusAudio/Source/OculusAudio/Private/OculusAudioDllManager.h
--- a/Engine/Plugins/Runtime/Oculus/OculusAudio/Source/OculusAudio/Private/OculusAudioDllManager.h
+++ b/Engine/Plugins/Runtime/Oculus/OculusAudio/Source/OculusAudio/Private/OculusAudioDllManager.h
@@ -37,14 +37,15 @@ static const TCHAR* GetOculusErrorString(ovrResult Result)
 #endif
 
 
-#define OVR_AUDIO_CHECK(Result, Context)																\
+#define OVR_AUDIO_CHECK_RETURN_VALUE(Result, Context, FailureReturnValue)								\
 	if (Result != ovrSuccess)																			\
 	{																									\
 		const TCHAR* ErrString = GetOculusErrorString(Result);											\
 		UE_LOG(LogAudio, Error, TEXT("Oculus Audio SDK Error - %s: %s"), TEXT(Context), ErrString);		\
-		return;																							\
+		return FailureReturnValue;																		\
 	}
 
+#define OVR_AUDIO_CHECK(Result, Context) OVR_AUDIO_CHECK_RETURN_VALUE(Result, Context, (void)0)
 
 /************************************************************************/
 /* FOculusAudioLibraryManager										   */
diff --git a/Engine/Plugins/Runtime/Oculus/OculusAudio/Source/OculusAudio/Private/OculusAudioMixer.h b/Engine/Plugins/Runtime/Oculus/OculusAudio/Source/OculusAudio/Private/OculusAudioMixer.h
--- a/Engine/Plugins/Runtime/Oculus/OculusAudio/Source/OculusAudio/Private/OculusAudioMixer.h
+++ b/Engine/Plugins/Runtime/Oculus/OculusAudio/Source/OculusAudio/Private/OculusAudioMixer.h
@@ -19,6 +19,7 @@ public:
 	~OculusAudioSpatializationAudioMixer();
 
 	void SetContext(ovrAudioContext* SharedContext);
+	void ClearContext();
 	ovrAudioContext* GetContext() { return Context; }
 
 	//~ Begin IAudioSpatialization 
diff --git a/Engine/Plugins/Runtime/Oculus/OculusAudio/Source/OculusAudio/Private/OculusAudioReverb.cpp b/Engine/Plugins/Runtime/Oculus/OculusAudio/Source/OculusAudio/Private/OculusAudioReverb.cpp
--- a/Engine/Plugins/Runtime/Oculus/OculusAudio/Source/OculusAudio/Private/OculusAudioReverb.cpp
+++ b/Engine/Plugins/Runtime/Oculus/OculusAudio/Source/OculusAudio/Private/OculusAudioReverb.cpp
@@ -10,6 +10,11 @@ void FSubmixEffectOculusReverbPlugin::SetContext(ovrAudioContext* SharedContext)
 	Context = SharedContext;
 }
 
+void FSubmixEffectOculusReverbPlugin::ClearContext()
+{
+	Context = nullptr;
+}
+
 FSubmixEffectOculusReverbPlugin::FSubmixEffectOculusReverbPlugin()
 	: Context(nullptr)
 {
@@ -46,6 +51,15 @@ void OculusAudioReverb::SetContext(ovrAudioContext* SharedContext)
 	}
 }
 
+void OculusAudioReverb::ClearContext()
+{
+	Context = nullptr;
+	for (FSubmixEffectOculusReverbPlugin* Submix : Submixes)
+	{
+		Submix->ClearContext();
+	}
+}
+
 FSoundEffectSubmix* OculusAudioReverb::GetEffectSubmix(class USoundSubmix* Submix)
 {
 	FSubmixEffectOculusReverbPlugin* ReverbSubmix = new FSubmixEffectOculusReverbPlugin();
diff --git a/Engine/Plugins/Runtime/Oculus/OculusAudio/Source/OculusAudio/Private/OculusAudioReverb.h b/Engine/Plugins/Runtime/Oculus/OculusAudio/Source/OculusAudio/Private/OculusAudioReverb.h
--- a/Engine/Plugins/Runtime/Oculus/OculusAudio/Source/OculusAudio/Private/OculusAudioReverb.h
+++ b/Engine/Plugins/Runtime/Oculus/OculusAudio/Source/OculusAudio/Private/OculusAudioReverb.h
@@ -13,6 +13,7 @@ public:
 	FSubmixEffectOculusReverbPlugin();
 
 	void SetContext(ovrAudioContext* SharedContext);
+	void ClearContext();
 
 	virtual void Init(const FSoundEffectSubmixInitData& InSampleRate) override
 	{
@@ -48,6 +49,7 @@ public:
 	}
 
 	void SetContext(ovrAudioContext* SharedContext);
+	void ClearContext();
 
 	virtual void OnInitSource(const uint32 SourceId, const FName& AudioComponentUserId, const uint32 NumChannels, UReverbPluginSourceSettingsBase* InSettings) override
 	{
diff --git a/Engine/Plugins/Runtime/Oculus/OculusAvatar/Source/Public/OvrAvatar.cpp b/Engine/Plugins/Runtime/Oculus/OculusAvatar/Source/Public/OvrAvatar.cpp
--- a/Engine/Plugins/Runtime/Oculus/OculusAvatar/Source/Public/OvrAvatar.cpp
+++ b/Engine/Plugins/Runtime/Oculus/OculusAvatar/Source/Public/OvrAvatar.cpp
@@ -2020,6 +2020,7 @@ ovrAvatarControllerType UOvrAvatar::GetControllerTypeByHardware()
 		controllerType = ovrAvatarControllerType_Go;
 		break;
 	case ovrpSystemHeadset_Oculus_Quest:
+	case ovrpSystemHeadset_Rift_S:
 		controllerType = ovrAvatarControllerType_Quest;
 		break;
 	case ovrpSystemHeadset_Rift_DK1:
@@ -2069,17 +2070,19 @@ bool UOvrAvatar::Is3DOFHardware()
 	return false;
 }
 
-void UOvrAvatar::UpdateVisemeValues(const TArray<float>& visemes)
+void UOvrAvatar::UpdateVisemeValues(const TArray<float>& visemes, const float laughterScore)
 {
 	FScopeLock Lock(&VisemeMutex);
 
-	const int32 VisemeCount = FMath::Min(OVR_AVATAR_MAXIMUM_VISEME_PARAM_COUNT, visemes.Num());
+	int32 VisemeCount = FMath::Min(OVR_AVATAR_MAXIMUM_VISEME_PARAM_COUNT - 1, visemes.Num());
 
-	VisemeValues.visemeParamCount = VisemeCount;
 	for (int i = 0; i < VisemeCount; i++)
 	{
 		VisemeValues.visemeParams[i] = visemes[i] * GAvatarVisemeMultiplier;
 	}
+
+	VisemeValues.visemeParams[VisemeCount] = laughterScore * GAvatarVisemeMultiplier;
+	VisemeValues.visemeParamCount = VisemeCount + 1;
 }
 
 void UOvrAvatar::UpdateHeadGazeTarget()
diff --git a/Engine/Plugins/Runtime/Oculus/OculusAvatar/Source/Public/OvrAvatar.h b/Engine/Plugins/Runtime/Oculus/OculusAvatar/Source/Public/OvrAvatar.h
--- a/Engine/Plugins/Runtime/Oculus/OculusAvatar/Source/Public/OvrAvatar.h
+++ b/Engine/Plugins/Runtime/Oculus/OculusAvatar/Source/Public/OvrAvatar.h
@@ -75,7 +75,7 @@ public:
 
 	void SetVoiceVisualValue(float value) { VoiceVisualValue = FMath::Clamp(value, 0.f, 1.f); }
 
-	void UpdateVisemeValues(const TArray<float>& visemes);
+	void UpdateVisemeValues(const TArray<float>& visemes, const float laughterScore);
 
 	void SetBodyCapability(bool Enable)	{ EnableBody = Enable; }
 	void SetBaseCapability(bool Enable) { EnableBase = Enable; }
