Fix up header issues that result in IncludeTool errors
#rnx

diff --git a/Engine/Plugins/Runtime/Oculus/OculusAudio/Source/OculusAudio/Private/OculusAudioLegacy.h b/Engine/Plugins/Runtime/Oculus/OculusAudio/Source/OculusAudio/Private/OculusAudioLegacy.h
--- a/Engine/Plugins/Runtime/Oculus/OculusAudio/Source/OculusAudio/Private/OculusAudioLegacy.h
+++ b/Engine/Plugins/Runtime/Oculus/OculusAudio/Source/OculusAudio/Private/OculusAudioLegacy.h
@@ -7,7 +7,6 @@
 #if PLATFORM_WINDOWS
 
 #include "OVR_Audio.h"
-#include "XAudio2Device.h"
 #include "AudioEffect.h"
 #include "XAudio2Effects.h"
 #include "OculusAudioEffect.h"
