Compile fixes round 2 for reverb plugin implementation

diff --git a/Engine/Plugins/Runtime/Oculus/OculusAudio/Source/OculusAudio/Private/OculusAudioReverb.cpp b/Engine/Plugins/Runtime/Oculus/OculusAudio/Source/OculusAudio/Private/OculusAudioReverb.cpp
--- a/Engine/Plugins/Runtime/Oculus/OculusAudio/Source/OculusAudio/Private/OculusAudioReverb.cpp
+++ b/Engine/Plugins/Runtime/Oculus/OculusAudio/Source/OculusAudio/Private/OculusAudioReverb.cpp
@@ -7,6 +7,14 @@
 #include "Sound/SoundSubmix.h"
 
 
+namespace
+{
+	TSharedPtr<FSubmixEffectOculusReverbPlugin, ESPMode::ThreadSafe> CastEffectToPluginSharedPtr(FSoundEffectSubmixPtr InSubmixEffect)
+	{
+		return StaticCastSharedPtr<FSubmixEffectOculusReverbPlugin, FSoundEffectSubmix, ESPMode::ThreadSafe>(InSubmixEffect);
+	}
+} // namespace <>
+
 void FSubmixEffectOculusReverbPlugin::SetContext(ovrAudioContext* SharedContext)
 {
 	FScopeLock ScopeLock(&ContextLock);
@@ -48,30 +56,30 @@ void OculusAudioReverb::SetContext(ovrAudioContext* SharedContext)
 		Context = SharedContext;
 	}
 
-	if (Submix.IsValid())
+	if (SubmixEffect.IsValid())
 	{
-		Submix->SetContext(SharedContext);
+		CastEffectToPluginSharedPtr(SubmixEffect)->SetContext(SharedContext);
 	}
 }
 
 void OculusAudioReverb::ClearContext()
 {
 	Context = nullptr;
-	if (Submix.IsValid())
+	if (SubmixEffect.IsValid())
 	{
-		Submix->ClearContext();
+		CastEffectToPluginSharedPtr(SubmixEffect)->ClearContext();
 	}
 }
 
 FSoundEffectSubmixPtr OculusAudioReverb::GetEffectSubmix()
 {
-	if (!Submix.IsValid())
+	if (!SubmixEffect.IsValid())
 	{
-		Submix = MakeShared<FSubmixEffectOculusReverbPlugin>();
-		Submix->SetEnabled(true);
+		SubmixEffect = MakeShareable(static_cast<FSoundEffectSubmix*>(new FSubmixEffectOculusReverbPlugin()));
+		SubmixEffect->SetEnabled(true);
 	}
 
-	return StaticCastSharedPtr<FSoundEffectSubmix, FSubmixEffectOculusReverbPlugin, ESPMode::ThreadSafe>(Submix);
+	return SubmixEffect;
 }
 
 USoundSubmix* OculusAudioReverb::GetSubmix()
diff --git a/Engine/Plugins/Runtime/Oculus/OculusAudio/Source/OculusAudio/Private/OculusAudioReverb.h b/Engine/Plugins/Runtime/Oculus/OculusAudio/Source/OculusAudio/Private/OculusAudioReverb.h
--- a/Engine/Plugins/Runtime/Oculus/OculusAudio/Source/OculusAudio/Private/OculusAudioReverb.h
+++ b/Engine/Plugins/Runtime/Oculus/OculusAudio/Source/OculusAudio/Private/OculusAudioReverb.h
@@ -72,5 +72,5 @@ public:
 	}
 private:
 	ovrAudioContext* Context;
-	TSharedPtr<FSubmixEffectOculusReverbPlugin> Submix;
+	FSoundEffectSubmixPtr SubmixEffect;
 };
\ No newline at end of file
diff --git a/Engine/Plugins/Runtime/Steam/SteamAudio/Source/SteamAudio/Private/PhononReverb.cpp b/Engine/Plugins/Runtime/Steam/SteamAudio/Source/SteamAudio/Private/PhononReverb.cpp
--- a/Engine/Plugins/Runtime/Steam/SteamAudio/Source/SteamAudio/Private/PhononReverb.cpp
+++ b/Engine/Plugins/Runtime/Steam/SteamAudio/Source/SteamAudio/Private/PhononReverb.cpp
@@ -424,7 +424,7 @@ namespace SteamAudio
 		}
 
 		bool bFoundPreset = false;
-		for (USubmixEffectReverbPluginPreset* Preset : ReverbSubmix->SubmixEffectChain)
+		for (USoundEffectSubmixPreset* Preset : ReverbSubmix->SubmixEffectChain)
 		{
 			if (USubmixEffectReverbPluginPreset* PluginPreset = Cast<USubmixEffectReverbPluginPreset>(Preset))
 			{
