Editgrate 4.25 audio features from project stream:
-Soundfield Submixes
-Endpoint Submixes
-Unreal Ambisonics Encoder/Decoder

diff --git a/Engine/Plugins/Runtime/AudioModulation/Source/AudioModulation/AudioModulation.Build.cs b/Engine/Plugins/Runtime/AudioModulation/Source/AudioModulation/AudioModulation.Build.cs
--- a/Engine/Plugins/Runtime/AudioModulation/Source/AudioModulation/AudioModulation.Build.cs
+++ b/Engine/Plugins/Runtime/AudioModulation/Source/AudioModulation/AudioModulation.Build.cs
@@ -9,7 +9,8 @@ namespace UnrealBuildTool.Rules
 			PublicDependencyModuleNames.AddRange(
 				new string[]
 				{
-					"Core"
+					"Core",
+					"AudioExtensions"
 				}
 			);
 
diff --git a/Engine/Plugins/Runtime/MicrosoftSpatialAudio/Source/MicrosoftSpatialSound/MicrosoftSpatialSound.Build.cs b/Engine/Plugins/Runtime/MicrosoftSpatialAudio/Source/MicrosoftSpatialSound/MicrosoftSpatialSound.Build.cs
--- a/Engine/Plugins/Runtime/MicrosoftSpatialAudio/Source/MicrosoftSpatialSound/MicrosoftSpatialSound.Build.cs
+++ b/Engine/Plugins/Runtime/MicrosoftSpatialAudio/Source/MicrosoftSpatialSound/MicrosoftSpatialSound.Build.cs
@@ -25,7 +25,8 @@ public class MicrosoftSpatialSound : ModuleRules
 			{
 				"Core",
 				"MixedRealityInteropLibrary",
-			}
+                "AudioExtensions"
+            }
 			);
 			
 		
@@ -33,7 +34,7 @@ public class MicrosoftSpatialSound : ModuleRules
 			new string[]
 			{
 				"CoreUObject",
-				"Engine"
+				"Engine",
 			}
 			);
         PrivateIncludePathModuleNames.Add("TargetPlatform");
diff --git a/Engine/Plugins/Runtime/MicrosoftSpatialAudio/Source/MicrosoftSpatialSound/Public/MicrosoftSpatialSoundPlugin.h b/Engine/Plugins/Runtime/MicrosoftSpatialAudio/Source/MicrosoftSpatialSound/Public/MicrosoftSpatialSoundPlugin.h
--- a/Engine/Plugins/Runtime/MicrosoftSpatialAudio/Source/MicrosoftSpatialSound/Public/MicrosoftSpatialSoundPlugin.h
+++ b/Engine/Plugins/Runtime/MicrosoftSpatialAudio/Source/MicrosoftSpatialSound/Public/MicrosoftSpatialSoundPlugin.h
@@ -8,7 +8,8 @@
 #include "DSP/Dsp.h"
 #include "HAL/Runnable.h"
 #include "HAL/RunnableThread.h"
-
+#include "Templates/Function.h"
+#include "Containers/Queue.h"
 #include "MixedRealityInterop.h"
 
 // Struct to hold dynamic object data for Microsoft Spatial Sound API
@@ -134,4 +135,4 @@ public:
 
 private:
 	FMicrosoftSpatialSoundPluginFactory PluginFactory;
-};
\ No newline at end of file
+};
diff --git a/Engine/Plugins/Runtime/Oculus/OculusAudio/Source/OculusAudio/Private/OculusAmbisonicSpatializer.cpp b/Engine/Plugins/Runtime/Oculus/OculusAudio/Source/OculusAudio/Private/OculusAmbisonicSpatializer.cpp
--- a/Engine/Plugins/Runtime/Oculus/OculusAudio/Source/OculusAudio/Private/OculusAmbisonicSpatializer.cpp
+++ b/Engine/Plugins/Runtime/Oculus/OculusAudio/Source/OculusAudio/Private/OculusAmbisonicSpatializer.cpp
@@ -1,187 +1,383 @@
 // Copyright Epic Games, Inc. All Rights Reserved.
 
 #include "OculusAmbisonicSpatializer.h"
+#include "ISoundfieldFormat.h"
 #include "OculusAudioMixer.h"
+#include "DSP/BufferVectorOperations.h"
+#include "AudioMixerDevice.h"
 #include "OculusAudioDllManager.h"
+#include "OculusAudioContextManager.h"
+#include "SoundFieldRendering.h"
 
-FOculusAmbisonicsMixer::FOculusAmbisonicsMixer()
-	: Context(nullptr)
-	, SampleRate(0)
-{
-}
+static int32 UseSoundfieldSubmixForOculusAudioOutputCVar = 0;
+FAutoConsoleVariableRef CVarUseSoundfieldSubmixForOculusAudioOutput(
+	TEXT("au.oculus.UseSoundfieldSubmixForOculusAudioOutput"),
+	UseSoundfieldSubmixForOculusAudioOutputCVar,
+	TEXT("If set to 1, this will allow for audio sources to be sent to an Oculus Audio soundfield submix to be HRTF spatialized, rather than using the Oculus Audio Spatialization Plugin for the entire project.\n")
+	TEXT("0: Disable, >0: Enable"),
+	ECVF_Default);
 
-void FOculusAmbisonicsMixer::Shutdown()
+// the Oculus Audio renderer renders all inputs immediately to an interleaved stereo output.
+class FOculusSoundfieldBuffer : public ISoundfieldAudioPacket
 {
-}
+public:
+	// Interleaved binaural audio buffer.
+	Audio::AlignedFloatBuffer AudioBuffer;
+	int32 NumChannels;
 
-int32 FOculusAmbisonicsMixer::GetNumChannelsForAmbisonicsFormat(UAmbisonicsSubmixSettingsBase* InSettings)
-{
-	return 4;
-}
+	FOculusSoundfieldBuffer()
+		: NumChannels(0)
+	{}
 
-void FOculusAmbisonicsMixer::OnOpenEncodingStream(const uint32 StreamId, UAmbisonicsSubmixSettingsBase* InSettings)
-{
-	UOculusAmbisonicsSettings* Settings = CastChecked<UOculusAmbisonicsSettings>(InSettings);
+	virtual void Serialize(FArchive& Ar) override
+	{
+		Ar << AudioBuffer;
+		Ar << NumChannels;
+	}
 
-	ovrAudioAmbisonicStream NewStream = nullptr;
-	ovrAudioAmbisonicFormat StreamFormat = (Settings->ChannelOrder == EAmbisonicFormat::AmbiX ? ovrAudioAmbisonicFormat_AmbiX : ovrAudioAmbisonicFormat_FuMa);
-	ovrResult OurResult = OVRA_CALL(ovrAudio_CreateAmbisonicStream)(Context, SampleRate, BufferLength, StreamFormat, 1, &NewStream);
-	check(OurResult == 0);
-	check(NewStream != nullptr);
-	OpenStreams.Add(StreamId, NewStream);
-}
 
-void FOculusAmbisonicsMixer::OnCloseEncodingStream(const uint32 SourceId)
-{
-	ovrAudioAmbisonicStream* InStream = OpenStreams.Find(SourceId);
-	if (InStream != nullptr)
+	virtual TUniquePtr<ISoundfieldAudioPacket> Duplicate() const override
 	{
-		OVRA_CALL(ovrAudio_DestroyAmbisonicStream)(*InStream);
-		OpenStreams.Remove(SourceId);
+		return TUniquePtr<ISoundfieldAudioPacket>(new FOculusSoundfieldBuffer(*this));
 	}
-	else
+
+
+	virtual void Reset() override
 	{
-		checkf(false, TEXT("Tried to close a stream we did not open."));
+		AudioBuffer.Reset();
+		NumChannels = 0;
 	}
-}
 
-void FOculusAmbisonicsMixer::EncodeToAmbisonics(const uint32 SourceId, const FAmbisonicsEncoderInputData& InputData, FAmbisonicsEncoderOutputData& OutputData, UAmbisonicsSubmixSettingsBase* InParams)
+};
+
+class FOculusEncoder : public ISoundfieldEncoderStream
 {
-	//Oculus only supports encoding mono streams.
-	if (InputData.NumChannels == 1)
+private:
+	ovrAudioContext Context;
+	TArray<int32> SourceIds;
+
+	// temp buffers use to binauralize each source.
+	Audio::AlignedFloatBuffer ScratchMonoBuffer;
+	Audio::AlignedFloatBuffer ScratchOutputBuffer;
+
+	static int32 SourceIdCounter;
+	static FCriticalSection SourceIdCounterCritSection;
+
+public:
+
+	FOculusEncoder(ovrAudioContext InContext, int32 InNumChannels, int32 MaxNumSources)
+		: Context(InContext)
 	{
-		//TODO: Implement;
-		//ovrAudio_MonoToAmbisonic(InputData.AudioBuffer->GetData())
+		check(Context != nullptr);
+		{
+			// Assign source Ids for each channel here.
+			FScopeLock ScopeLock(&SourceIdCounterCritSection);
+			for (int32 Index = 0; Index < InNumChannels; Index++)
+			{
+
+				SourceIds.Add(SourceIdCounter);
+				SourceIdCounter = (SourceIdCounter + 1) % MaxNumSources;
+			}
+		}
+
+		// Set default settings for each source channel here.
+		for (int32& SourceId : SourceIds)
+		{
+			uint32 Flags = 0;
+			Flags |= ovrAudioSourceFlag_ReflectionsDisabled;
+
+			ovrResult Result = OVRA_CALL(ovrAudio_SetAudioSourceFlags)(Context, SourceId, Flags);
+			check(Result == ovrSuccess);
+
+			ovrAudioSourceAttenuationMode mode = ovrAudioSourceAttenuationMode_None;
+			Result = OVRA_CALL(ovrAudio_SetAudioSourceAttenuationMode)(Context, SourceId, mode, 1.0f);
+			check(Result == ovrSuccess);
+
+			Result = OVRA_CALL(ovrAudio_SetAudioReverbSendLevel)(Context, SourceId, 0.0f);
+			check(Result == ovrSuccess);
+		}
 	}
-}
 
-void FOculusAmbisonicsMixer::OnOpenDecodingStream(const uint32 StreamId, UAmbisonicsSubmixSettingsBase* InSettings, FAmbisonicsDecoderPositionalData& SpecifiedOutputPositions)
-{
-	ovrAudioAmbisonicFormat StreamFormat = ovrAudioAmbisonicFormat_AmbiX;
-	ovrAudioAmbisonicSpeakerLayout SpeakerLayout = ovrAudioAmbisonicSpeakerLayout_SphericalHarmonics;
-	if (InSettings != nullptr) 
+	virtual void Encode(const FSoundfieldEncoderInputData& InputData, ISoundfieldAudioPacket& OutputData) override
 	{
-		UOculusAmbisonicsSettings* Settings = CastChecked<UOculusAmbisonicsSettings>(InSettings);
-		StreamFormat = (Settings->ChannelOrder == EAmbisonicFormat::AmbiX ? ovrAudioAmbisonicFormat_AmbiX : ovrAudioAmbisonicFormat_FuMa);
-		SpeakerLayout = (Settings->SpatializationMode == EAmbisonicMode::SphericalHarmonics) ?
-			ovrAudioAmbisonicSpeakerLayout_SphericalHarmonics : ovrAudioAmbisonicSpeakerLayout_Icosahedron;
+		OutputData.Reset();
+		EncodeAndMixIn(InputData, OutputData);
 	}
 
-	ovrAudioAmbisonicStream NewStream = nullptr;
-	ovrResult OurResult = OVRA_CALL(ovrAudio_CreateAmbisonicStream)(Context, SampleRate, BufferLength, StreamFormat, 1, &NewStream);
 
-	check(OurResult == 0);
-	check(NewStream != nullptr);
+	// Binaurally spatialize each independent speaker position. If UseSoundfieldSubmixForOculusAudioOutputCVar
+	// was false when this encoder was created this will no-op.
+	virtual void EncodeAndMixIn(const FSoundfieldEncoderInputData& InputData, ISoundfieldAudioPacket& OutputData) override
+	{
+		FOculusSoundfieldBuffer& OutputBuffer = DowncastSoundfieldRef<FOculusSoundfieldBuffer>(OutputData);
 
-	OurResult = OVRA_CALL(ovrAudio_SetAmbisonicSpeakerLayout)(NewStream, SpeakerLayout);
+		check(InputData.NumChannels == SourceIds.Num());
+		const int32 NumFrames = InputData.AudioBuffer.Num() / InputData.NumChannels;
 
-	check(OurResult == 0);
-	OpenStreams.Add(StreamId, NewStream);
-}
+		OutputBuffer.AudioBuffer.Reset();
+		OutputBuffer.NumChannels = 2;
+		OutputBuffer.AudioBuffer.AddZeroed(NumFrames * OutputBuffer.NumChannels);
+
+		FQuat ListenerOrientation = InputData.PositionalData.Rotation.Inverse();
+
+		// Translate the input position to OVR coordinates
+		FVector ListenerPosition = OculusAudioSpatializationAudioMixer::ToOVRVector(FVector::ZeroVector);
+		FVector ListenerForward = OculusAudioSpatializationAudioMixer::ToOVRVector(ListenerOrientation.GetForwardVector());
+		FVector OvrListenerUp = OculusAudioSpatializationAudioMixer::ToOVRVector(ListenerOrientation.GetUpVector());
+
+		ovrResult Result = OVRA_CALL(ovrAudio_SetListenerVectors)(Context,
+			ListenerPosition.X, ListenerPosition.Y, ListenerPosition.Z,
+			ListenerForward.X, ListenerForward.Y, ListenerForward.Z,
+			OvrListenerUp.X, OvrListenerUp.Y, OvrListenerUp.Z);
+		OVR_AUDIO_CHECK(Result, "Failed to set listener position and rotation");
+
+		check(InputData.PositionalData.ChannelPositions);
+		const TArray<Audio::FChannelPositionInfo>& ChannelPositions = *(InputData.PositionalData.ChannelPositions);
+
+		for (int32 ChannelIndex = 0; ChannelIndex < SourceIds.Num(); ChannelIndex++)
+		{
+			int32& SourceId = SourceIds[ChannelIndex];
+			const Audio::FChannelPositionInfo& ChannelPosition = ChannelPositions[ChannelIndex];
+			
+
+			// Translate the input position to OVR coordinates
+			FVector SourcePosition = OculusAudioSpatializationAudioMixer::ToOVRVector(ChannelPosition);
+
+			// Set the source position to current audio position
+			Result = OVRA_CALL(ovrAudio_SetAudioSourcePos)(Context, SourceId, SourcePosition.X, SourcePosition.Y, SourcePosition.Z);
+			OVR_AUDIO_CHECK(Result, "Failed to set audio source position");
+
+			// Deinterleave the audio into the mono temp buffer.
+			ScratchMonoBuffer.Reset();
+			ScratchMonoBuffer.AddUninitialized(NumFrames);
+			for(int32 FrameIndex = 0; FrameIndex < NumFrames; FrameIndex++)
+			{
+				ScratchMonoBuffer[FrameIndex] = InputData.AudioBuffer[FrameIndex * InputData.NumChannels + ChannelIndex];
+			}
 
-void FOculusAmbisonicsMixer::OnCloseDecodingStream(const uint32 StreamId)
+			ScratchOutputBuffer.Reset();
+			ScratchOutputBuffer.AddUninitialized(NumFrames * 2);
+
+			uint32 Status = 0;
+			Result = OVRA_CALL(ovrAudio_SpatializeMonoSourceInterleaved)(Context, SourceId, &Status, ScratchOutputBuffer.GetData(), ScratchMonoBuffer.GetData());
+			OVR_AUDIO_CHECK(Result, "Failed to spatialize mono source interleaved");
+
+			// sum the interleaved output into the output bed.
+			Audio::MixInBufferFast(ScratchOutputBuffer, OutputBuffer.AudioBuffer);
+		}
+	}
+};
+
+int32 FOculusEncoder::SourceIdCounter = 0;
+FCriticalSection FOculusEncoder::SourceIdCounterCritSection;
+
+// Because Oculus spatializes every directly to an interleaved stereo buffer,
+// we simply pass it forward here.
+class FOculusAudioDecoder : public ISoundfieldDecoderStream
 {
-	ovrAudioAmbisonicStream* InStream = OpenStreams.Find(StreamId);
-	if (InStream != nullptr)
+public:
+
+	virtual void Decode(const FSoundfieldDecoderInputData& InputData, FSoundfieldDecoderOutputData& OutputData) override
 	{
-		OVRA_CALL(ovrAudio_DestroyAmbisonicStream)(*InStream);
-		OpenStreams.Remove(StreamId);
+		OutputData.AudioBuffer.Reset();
+		OutputData.AudioBuffer.AddZeroed(InputData.NumFrames * InputData.PositionalData.NumChannels);
+		DecodeAndMixIn(InputData, OutputData);
 	}
-	else
+
+	virtual void DecodeAndMixIn(const FSoundfieldDecoderInputData& InputData, FSoundfieldDecoderOutputData& OutputData) override
 	{
-		checkf(false, TEXT("Tried to close a stream we did not open."));
+		if (UseSoundfieldSubmixForOculusAudioOutputCVar)
+		{
+			const FOculusSoundfieldBuffer& InputBuffer = DowncastSoundfieldRef<const FOculusSoundfieldBuffer>(InputData.SoundfieldBuffer);
+
+			if (InputBuffer.AudioBuffer.Num() == 0)
+			{
+				return;
+			}
+			else if (InputData.PositionalData.NumChannels == InputBuffer.NumChannels)
+			{
+				// If the number of output channels is the same as the input channels, mix it in directly.
+				Audio::MixInBufferFast(InputBuffer.AudioBuffer, OutputData.AudioBuffer);
+			}
+			else
+			{
+				// Otherwise, downmix and mix in.
+				Audio::AlignedFloatBuffer OutChannelMap;
+				Audio::FMixerDevice::Get2DChannelMap(false, InputBuffer.NumChannels, InputData.PositionalData.NumChannels, false, OutChannelMap);
+				Audio::DownmixAndSumIntoBuffer(InputBuffer.NumChannels, InputData.PositionalData.NumChannels, InputBuffer.AudioBuffer, OutputData.AudioBuffer, OutChannelMap.GetData());
+			}
+		}
 	}
-}
+};
 
-void FOculusAmbisonicsMixer::DecodeFromAmbisonics(const uint32 StreamId, const FAmbisonicsDecoderInputData& InputData, FAmbisonicsDecoderPositionalData& SpecifiedOutputPositions, FAmbisonicsDecoderOutputData& OutputData)
+// Class that takes ambisonics audio and renders it to an FOculusSoundfieldBuffer.
+class FOculusAmbisonicsTranscoder : public ISoundfieldTranscodeStream
 {
-	//Currently, Oculus only decodes first-order to stereo.
-	if (SpecifiedOutputPositions.OutputNumChannels == 2 && InputData.NumChannels == 4)
+private:
+	ovrAudioContext Context;
+	ovrAudioAmbisonicStream AmbiStreamObject;
+public:
+	FOculusAmbisonicsTranscoder(ovrAudioContext InContext, int32 InSampleRate, int32 InBufferLength)
+		: Context(InContext)
 	{
-		//Get ambisonics stream
-		ovrAudioAmbisonicStream* ThisStream = OpenStreams.Find(StreamId);
-		if (ThisStream != nullptr)
+		ovrAudioAmbisonicFormat StreamFormat = ovrAudioAmbisonicFormat_AmbiX;
+		ovrResult OurResult = OVRA_CALL(ovrAudio_CreateAmbisonicStream)(Context, InSampleRate, InBufferLength, StreamFormat, 1, &AmbiStreamObject);
+		check(OurResult == 0);
+		check(AmbiStreamObject != nullptr);
+	}
+
+	virtual void Transcode(const ISoundfieldAudioPacket& InputData, const ISoundfieldEncodingSettingsProxy& InputSettings, ISoundfieldAudioPacket& OutputData, const ISoundfieldEncodingSettingsProxy& OutputSettings) override
+	{
+		OutputData.Reset();
+		TranscodeAndMixIn(InputData, InputSettings, OutputData, OutputSettings);
+	}
+
+
+	virtual void TranscodeAndMixIn(const ISoundfieldAudioPacket& InputData, const ISoundfieldEncodingSettingsProxy& InputSettings, ISoundfieldAudioPacket& PacketToSumTo, const ISoundfieldEncodingSettingsProxy& OutputSettings) override
+	{
+		const FAmbisonicsSoundfieldBuffer& InputBuffer = DowncastSoundfieldRef<const FAmbisonicsSoundfieldBuffer>(InputData);
+		FOculusSoundfieldBuffer& OutputBuffer = DowncastSoundfieldRef<FOculusSoundfieldBuffer>(PacketToSumTo);
+
+		const int32 NumFrames = InputBuffer.AudioBuffer.Num() / InputBuffer.NumChannels;
+		OutputBuffer.AudioBuffer.SetNumZeroed(NumFrames * 2);
+		OutputBuffer.NumChannels = 2;
+
+		//Currently, Oculus only decodes first-order ambisonics to stereo.
+		// in the future we can truncate InputBuffer to four channels and 
+		// render that.
+		// For now we just output silence for higher order ambisonics from the Oculus spatializaers
+		if (InputBuffer.NumChannels == 4)
 		{
-			//convert ambisonics rotation to listener rotation:
-			{
-				// Translate the input position to OVR coordinates
-				FQuat& ListenerRotation = SpecifiedOutputPositions.ListenerRotation;
-				FVector OvrListenerForward = OculusAudioSpatializationAudioMixer::ToOVRVector(ListenerRotation.GetForwardVector());
-				FVector OvrListenerUp = OculusAudioSpatializationAudioMixer::ToOVRVector(ListenerRotation.GetUpVector());
-
-				static auto SetListenerVectors = OVRA_CALL(ovrAudio_SetListenerVectors);
-				ovrResult Result = SetListenerVectors(Context,
-					0.0f, 0.0f, 0.0f,
-					OvrListenerForward.X, OvrListenerForward.Y, OvrListenerForward.Z,
-					OvrListenerUp.X, OvrListenerUp.Y, OvrListenerUp.Z);
-			}
-			ovrResult DecodeResult= OVRA_CALL(ovrAudio_ProcessAmbisonicStreamInterleaved)(Context, *ThisStream, InputData.AudioBuffer->GetData(), OutputData.AudioBuffer.GetData(), InputData.AudioBuffer->Num() / InputData.NumChannels);
+			// Translate the input position to OVR coordinates
+			const FQuat& ListenerRotation = InputBuffer.Rotation;
+			FVector OvrListenerForward = OculusAudioSpatializationAudioMixer::ToOVRVector(ListenerRotation.GetForwardVector());
+			FVector OvrListenerUp = OculusAudioSpatializationAudioMixer::ToOVRVector(ListenerRotation.GetUpVector());
+
+			static auto SetListenerVectors = OVRA_CALL(ovrAudio_SetListenerVectors);
+			ovrResult Result = SetListenerVectors(Context,
+				0.0f, 0.0f, 0.0f,
+				OvrListenerForward.X, OvrListenerForward.Y, OvrListenerForward.Z,
+				OvrListenerUp.X, OvrListenerUp.Y, OvrListenerUp.Z);
+
+			ovrResult DecodeResult = OVRA_CALL(ovrAudio_ProcessAmbisonicStreamInterleaved)(Context, AmbiStreamObject, InputBuffer.AudioBuffer.GetData(), OutputBuffer.AudioBuffer.GetData(), NumFrames);
 			check(DecodeResult == 0);
 		}
+	}
+
+};
+
+// Since FOculusSoundfieldBuffer is just an interleaved stereo buffer, we simply mix the buffer here.
+class FOculusAudioSoundfieldMixer : public ISoundfieldMixerStream
+{
+public:
+	virtual void MixTogether(const FSoundfieldMixerInputData& InputData, ISoundfieldAudioPacket& PacketToSumTo) override
+	{
+		const FOculusSoundfieldBuffer& InputBuffer = DowncastSoundfieldRef<const FOculusSoundfieldBuffer>(InputData.InputPacket);
+		FOculusSoundfieldBuffer& OutputBuffer = DowncastSoundfieldRef<FOculusSoundfieldBuffer>(PacketToSumTo);
+
+		if (OutputBuffer.AudioBuffer.Num() == 0)
+		{
+			OutputBuffer.AudioBuffer = InputBuffer.AudioBuffer;
+			OutputBuffer.NumChannels = InputBuffer.NumChannels;
+		}
 		else
 		{
-			checkf(false, TEXT("Invalid stream ID."));
+			check(InputBuffer.AudioBuffer.Num() == OutputBuffer.AudioBuffer.Num());
+			check(InputBuffer.NumChannels == OutputBuffer.NumChannels);
+
+
+			Audio::MixInBufferFast(InputBuffer.AudioBuffer, OutputBuffer.AudioBuffer, InputData.SendLevel);
 		}
 	}
-	else
-	{
-		checkf(false, TEXT("Invalid number of channels. Input channels were %d but should be 4, output channels were %d but should be 2."), InputData.NumChannels, SpecifiedOutputPositions.OutputNumChannels);
-	}
+};
 
-#define AMBISONICS_SINE_TEST 0
+FOculusAmbisonicsFactory::FOculusAmbisonicsFactory()
+{
+	ISoundfieldFactory::RegisterSoundfieldFormat(this);
+}
+
+FOculusAmbisonicsFactory::~FOculusAmbisonicsFactory()
+{
+	ISoundfieldFactory::UnregisterSoundfieldFormat(this);
+}
+
+bool FOculusAmbisonicsFactory::IsTranscodeRequiredBetweenSettings(const ISoundfieldEncodingSettingsProxy& SourceSettings, const ISoundfieldEncodingSettingsProxy& DestinationSettings)
+{
+	return false;
+}
 
-#if AMBISONICS_SINE_TEST
-	static float n = 0.0f;
-	for (int32 Index = 0; Index < OutputData.AudioBuffer.Num(); Index+= SpecifiedOutputPositions.OutputNumChannels)
+bool FOculusAmbisonicsFactory::CanTranscodeFromSoundfieldFormat(FName SourceFormat, const ISoundfieldEncodingSettingsProxy& SourceEncodingSettings)
+{
+	FName AmbisonicsSoundfieldFormatName = GetUnrealAmbisonicsFormatName();
+	if (SourceFormat == AmbisonicsSoundfieldFormatName)
 	{
-		// Do a sine test before we're all done
-		for (int32 ChannelIndex = 0; ChannelIndex < SpecifiedOutputPositions.OutputNumChannels; ChannelIndex++)
-		{
-			OutputData.AudioBuffer[Index + ChannelIndex] = FMath::Sin(440.0f * n * 2.0f * PI / 48000.0f);
-		}
-		n += 1.0f;
+		return true;
 	}
-#endif
+	else
+	{
+		return false;
+	}
+}
+
+bool FOculusAmbisonicsFactory::CanTranscodeToSoundfieldFormat(FName DestinationFormat, const ISoundfieldEncodingSettingsProxy& DestinationEncodingSettings)
+{
+	return false;
 }
 
-bool FOculusAmbisonicsMixer::ShouldReencodeBetween(UAmbisonicsSubmixSettingsBase* SourceSubmixSettings, UAmbisonicsSubmixSettingsBase* DestinationSubmixSettings)
+UClass* FOculusAmbisonicsFactory::GetCustomEncodingSettingsClass() const
 {
-	//stub
-	return true;
+	return UOculusAudioSoundfieldSettings::StaticClass();
 }
 
-void FOculusAmbisonicsMixer::Initialize(const FAudioPluginInitializationParams InitializationParams)
+const USoundfieldEncodingSettingsBase* FOculusAmbisonicsFactory::GetDefaultEncodingSettings()
 {
-	SampleRate = InitializationParams.SampleRate;
-	BufferLength = InitializationParams.BufferLength;
+	return GetDefault<UOculusAudioSoundfieldSettings>();
+}
 
+FName FOculusAmbisonicsFactory::GetSoundfieldFormatName()
+{
+	static FName OculusBinauralFormatName = FName(TEXT("Oculus Binaural"));
+	return OculusBinauralFormatName;
+}
 
-	ovrAudioContextConfiguration ContextConfig;
-	ContextConfig.acc_Size = sizeof(ovrAudioContextConfiguration);
+TUniquePtr<ISoundfieldEncoderStream> FOculusAmbisonicsFactory::CreateEncoderStream(const FAudioPluginInitializationParams& InitInfo, const ISoundfieldEncodingSettingsProxy& InitialSettings)
+{
+	ovrAudioContext Context = FOculusAudioContextManager::GetContextForAudioDevice(InitInfo.AudioDevicePtr);
+	if (!Context)
+	{
+		Context = FOculusAudioContextManager::CreateContextForAudioDevice(InitInfo.AudioDevicePtr);
+		check(Context);
+	}
 
-	// TODO: Check if MaxNumSources also handles ambisonics stuff.
-	ContextConfig.acc_MaxNumSources = 1;
-	ContextConfig.acc_SampleRate = SampleRate;
-	ContextConfig.acc_BufferLength = BufferLength;
-	ovrResult Result = OVRA_CALL(ovrAudio_CreateContext)(&Context, &ContextConfig);
-	OVR_AUDIO_CHECK(Result, "Failed to create ambisonic context");
+	return TUniquePtr<ISoundfieldEncoderStream>(new FOculusEncoder(Context, InitInfo.NumOutputChannels, InitInfo.NumSources));
 }
 
-UClass* FOculusAmbisonicsMixer::GetCustomSettingsClass()
+TUniquePtr<ISoundfieldDecoderStream> FOculusAmbisonicsFactory::CreateDecoderStream(const FAudioPluginInitializationParams& InitInfo, const ISoundfieldEncodingSettingsProxy& InitialSettings)
 {
-	return UOculusAmbisonicsSettings::StaticClass();
+	return TUniquePtr<ISoundfieldDecoderStream>(new FOculusAudioDecoder());
 }
 
-UAmbisonicsSubmixSettingsBase* FOculusAmbisonicsMixer::GetDefaultSettings()
+TUniquePtr<ISoundfieldTranscodeStream> FOculusAmbisonicsFactory::CreateTranscoderStream(const FName SourceFormat, const ISoundfieldEncodingSettingsProxy& InitialSourceSettings, const FName DestinationFormat, const ISoundfieldEncodingSettingsProxy& InitialDestinationSettings, const FAudioPluginInitializationParams& InitInfo)
 {
-	static UOculusAmbisonicsSettings* DefaultSettingsPtr = nullptr;
-	if (DefaultSettingsPtr == nullptr)
+	check(SourceFormat == GetUnrealAmbisonicsFormatName());
+	check(DestinationFormat == GetSoundfieldFormatName());
+
+	ovrAudioContext Context = FOculusAudioContextManager::GetContextForAudioDevice(InitInfo.AudioDevicePtr);
+	if (!Context)
 	{
-		DefaultSettingsPtr = NewObject<UOculusAmbisonicsSettings>();
-		DefaultSettingsPtr->ChannelOrder = EAmbisonicFormat::AmbiX;
-		DefaultSettingsPtr->SpatializationMode = EAmbisonicMode::SphericalHarmonics;
-		DefaultSettingsPtr->AddToRoot();
+		Context = FOculusAudioContextManager::CreateContextForAudioDevice(InitInfo.AudioDevicePtr);
+		check(Context);
 	}
 
-	return DefaultSettingsPtr;
+	return TUniquePtr<ISoundfieldTranscodeStream>(new FOculusAmbisonicsTranscoder(Context, InitInfo.SampleRate, InitInfo.BufferLength));
+}
+
+TUniquePtr<ISoundfieldMixerStream> FOculusAmbisonicsFactory::CreateMixerStream(const ISoundfieldEncodingSettingsProxy& InitialSettings)
+{
+	return TUniquePtr<ISoundfieldMixerStream>(new FOculusAudioSoundfieldMixer());
+}
+
+TUniquePtr<ISoundfieldAudioPacket> FOculusAmbisonicsFactory::CreateEmptyPacket()
+{
+	return TUniquePtr<ISoundfieldAudioPacket>(new FOculusSoundfieldBuffer());
 }
diff --git a/Engine/Plugins/Runtime/Oculus/OculusAudio/Source/OculusAudio/Private/OculusAudio.cpp b/Engine/Plugins/Runtime/Oculus/OculusAudio/Source/OculusAudio/Private/OculusAudio.cpp
--- a/Engine/Plugins/Runtime/Oculus/OculusAudio/Source/OculusAudio/Private/OculusAudio.cpp
+++ b/Engine/Plugins/Runtime/Oculus/OculusAudio/Source/OculusAudio/Private/OculusAudio.cpp
@@ -27,11 +27,6 @@ TAudioSpatializationPtr FOculusSpatializationPluginFactory::CreateNewSpatializat
 }
 
 
-TAmbisonicsMixerPtr FOculusSpatializationPluginFactory::CreateNewAmbisonicsMixer(FAudioDevice* OwningDevice)
-{
-	return TAmbisonicsMixerPtr(new FOculusAmbisonicsMixer());
-}
-
 TAudioReverbPtr FOculusReverbPluginFactory::CreateNewReverbPlugin(FAudioDevice* OwningDevice)
 {
 	FOculusAudioPlugin* Plugin = &FModuleManager::GetModuleChecked<FOculusAudioPlugin>("OculusAudio");
diff --git a/Engine/Plugins/Runtime/Oculus/OculusAudio/Source/OculusAudio/Private/OculusAudio.h b/Engine/Plugins/Runtime/Oculus/OculusAudio/Source/OculusAudio/Private/OculusAudio.h
--- a/Engine/Plugins/Runtime/Oculus/OculusAudio/Source/OculusAudio/Private/OculusAudio.h
+++ b/Engine/Plugins/Runtime/Oculus/OculusAudio/Source/OculusAudio/Private/OculusAudio.h
@@ -8,6 +8,7 @@
 #include "OculusAudioLegacy.h"
 #endif
 #include "OculusAudioSourceSettings.h"
+#include "OculusAmbisonicSpatializer.h"
 
 /************************************************************************/
 /* FOculusSpatializationPluginFactory								   */
@@ -34,8 +35,6 @@ public:
 	virtual TAudioSpatializationPtr CreateNewSpatializationPlugin(FAudioDevice* OwningDevice) override;
 	//~ End IAudioSpatializationFactory
 
-	virtual TAmbisonicsMixerPtr CreateNewAmbisonicsMixer(FAudioDevice* OwningDevice) override;
-
 };
 
 class FOculusReverbPluginFactory : public IAudioReverbFactory
diff --git a/Engine/Plugins/Runtime/Oculus/OculusAudio/Source/OculusAudio/Private/OculusAudioContextManager.cpp b/Engine/Plugins/Runtime/Oculus/OculusAudio/Source/OculusAudio/Private/OculusAudioContextManager.cpp
--- a/Engine/Plugins/Runtime/Oculus/OculusAudio/Source/OculusAudio/Private/OculusAudioContextManager.cpp
+++ b/Engine/Plugins/Runtime/Oculus/OculusAudio/Source/OculusAudio/Private/OculusAudioContextManager.cpp
@@ -10,67 +10,20 @@
 ovrAudioContext FOculusAudioContextManager::SerializationContext = nullptr;
 UActorComponent* FOculusAudioContextManager::SerializationParent = nullptr;
 
+TMap<const FAudioDevice*, ovrAudioContext> FOculusAudioContextManager::ContextMap;
+FCriticalSection FOculusAudioContextManager::ContextCritSection;
+
 FOculusAudioContextManager::FOculusAudioContextManager()
-	: Context(nullptr)
 {
 	// empty
 }
 
 FOculusAudioContextManager::~FOculusAudioContextManager()
 {
-	OVRA_CALL(ovrAudio_DestroyContext)(Context);
-	Context = nullptr;
-	SerializationContext = nullptr;
 }
 
 void FOculusAudioContextManager::OnListenerInitialize(FAudioDevice* AudioDevice, UWorld* ListenerWorld)
 {
-	// Inject the Context into the spatailizer (and reverb) if they're enabled
-	FOculusAudioPlugin* Plugin = &FModuleManager::GetModuleChecked<FOculusAudioPlugin>("OculusAudio");
-	if (Plugin == nullptr)
-	{
-		return; // PAS is this needed?
-	}
-
-	if (SerializationContext != nullptr)
-	{
-#if 0
-		/* NOTE: this was meant to be some magic that would clean up the serialization
-		context so there is no stale state in PIE. This failed when there was no geometry
-		in the scene. Since the plugins are force to use one global context that seems
-		like a better option for the vanilla UE4 integraiton to do as well for consistency.
-		*/
-		// We must have created some geometry
-		check(SerializationParent != nullptr);
-		AActor* Actor = SerializationParent->GetOwner();
-		check(Actor != nullptr);
-		UWorld* SerializedWorld = Actor->GetWorld();
-		check(SerializedWorld == ListenerWorld);
-		Context = SerializationContext;
-		SerializationContext = nullptr;
-		SerializationParent = nullptr;
-#else
-		Context = SerializationContext;
-		SerializationContext = nullptr;
-#endif
-	}
-
-	FString OculusSpatializerPluginName = Plugin->GetSpatializationPluginFactory()->GetDisplayName();
-	FString CurrentSpatializerPluginName = AudioPluginUtilities::GetDesiredPluginName(EAudioPlugin::SPATIALIZATION);
-	if (CurrentSpatializerPluginName.Equals(OculusSpatializerPluginName)) // we have a match!
-	{
-		OculusAudioSpatializationAudioMixer* Spatializer = 
-			static_cast<OculusAudioSpatializationAudioMixer*>(AudioDevice->SpatializationPluginInterface.Get());
-		Spatializer->SetContext(&Context);
-	}
-
-	FString OculusReverbPluginName = Plugin->GetReverbPluginFactory()->GetDisplayName();
-	FString CurrentReverbPluginName = AudioPluginUtilities::GetDesiredPluginName(EAudioPlugin::REVERB);
-	if (CurrentReverbPluginName.Equals(OculusReverbPluginName))
-	{
-		OculusAudioReverb* Reverb = static_cast<OculusAudioReverb*>(AudioDevice->ReverbPluginInterface.Get());
-		Reverb->SetContext(&Context);
-	}
 }
 
 void FOculusAudioContextManager::OnListenerShutdown(FAudioDevice* AudioDevice)
@@ -84,7 +37,7 @@ void FOculusAudioContextManager::OnListenerShutdown(FAudioDevice* AudioDevice)
 	{
 		OculusAudioSpatializationAudioMixer* Spatializer =
 			static_cast<OculusAudioSpatializationAudioMixer*>(AudioDevice->SpatializationPluginInterface.Get());
-		Spatializer->SetContext(nullptr);
+		Spatializer->ClearContext();
 	}
 
 	FString OculusReverbPluginName = Plugin->GetReverbPluginFactory()->GetDisplayName();
@@ -92,8 +45,15 @@ void FOculusAudioContextManager::OnListenerShutdown(FAudioDevice* AudioDevice)
 	if (CurrentReverbPluginName.Equals(OculusReverbPluginName))
 	{
 		OculusAudioReverb* Reverb = static_cast<OculusAudioReverb*>(AudioDevice->ReverbPluginInterface.Get());
-		Reverb->SetContext(nullptr);
+		Reverb->ClearContext();
 	}
+
+	// FIXME:
+	// There's a possibility this will leak if a Oculus Binaural submix is created,
+	// but Oculus audio is not specified as the spatialization or reverb plugin.
+	// This is a niche use case, but could be solved by having the oculus soundfield explicitly destroy
+	// a context.
+	DestroyContextForAudioDevice(AudioDevice);
 }
 
 ovrAudioContext FOculusAudioContextManager::GetOrCreateSerializationContext(UActorComponent* Parent)
@@ -145,3 +105,51 @@ ovrAudioContext FOculusAudioContextManager::GetOrCreateSerializationContext(UAct
 
 	return SerializationContext;
 }
+
+ovrAudioContext FOculusAudioContextManager::GetContextForAudioDevice(const FAudioDevice* InAudioDevice)
+{
+	FScopeLock ScopeLock(&ContextCritSection);
+	ovrAudioContext* Context = ContextMap.Find(InAudioDevice);
+	if (Context)
+	{
+		return *Context;
+	}
+	else
+	{
+		return nullptr;
+	}
+}
+
+ovrAudioContext FOculusAudioContextManager::CreateContextForAudioDevice(FAudioDevice* InAudioDevice)
+{
+	ovrAudioContextConfiguration ContextConfig = { 0 };
+	ContextConfig.acc_BufferLength = InAudioDevice->GetBufferLength();
+	ContextConfig.acc_MaxNumSources = InAudioDevice->GetMaxSources();
+	ContextConfig.acc_SampleRate = InAudioDevice->GetSampleRate();
+	ContextConfig.acc_Size = sizeof(ovrAudioContextConfiguration);
+
+	ovrAudioContext NewContext = nullptr;
+	ovrResult Result = OVRA_CALL(ovrAudio_CreateContext)(&NewContext, &ContextConfig);
+
+	if(ensure(Result == ovrSuccess))
+	{
+		FScopeLock ScopeLock(&ContextCritSection);
+		return ContextMap.Add(InAudioDevice, NewContext);
+	}
+	else
+	{
+		return nullptr;
+	}
+}
+
+void FOculusAudioContextManager::DestroyContextForAudioDevice(const FAudioDevice* InAudioDevice)
+{
+	FScopeLock ScopeLock(&ContextCritSection);
+	ovrAudioContext* Context = ContextMap.Find(InAudioDevice);
+
+	if (Context)
+	{
+		OVRA_CALL(ovrAudio_DestroyContext)(*Context);
+		ContextMap.Remove(InAudioDevice);
+	}
+}
diff --git a/Engine/Plugins/Runtime/Oculus/OculusAudio/Source/OculusAudio/Private/OculusAudioGeometryComponent.cpp b/Engine/Plugins/Runtime/Oculus/OculusAudio/Source/OculusAudio/Private/OculusAudioGeometryComponent.cpp
--- a/Engine/Plugins/Runtime/Oculus/OculusAudio/Source/OculusAudio/Private/OculusAudioGeometryComponent.cpp
+++ b/Engine/Plugins/Runtime/Oculus/OculusAudio/Source/OculusAudio/Private/OculusAudioGeometryComponent.cpp
@@ -34,15 +34,6 @@ UOculusAudioGeometryComponent::UOculusAudioGeometryComponent()
 
 UOculusAudioGeometryComponent::~UOculusAudioGeometryComponent()
 {
-	if (CachedContext != nullptr)
-	{
-		ovrAudioContext PluginContext = FOculusAudioLibraryManager::Get().GetPluginContext();
-		if (PluginContext != nullptr)
-		{
-			check(*CachedContext == PluginContext);
-			delete CachedContext;
-		}
-	}
 }
 
 // ------------------------------------------------------------------------------------------------------------------------------------------------------------
@@ -306,8 +297,7 @@ ovrAudioContext UOculusAudioGeometryComponent::GetContext(UWorld* World)
 		ovrAudioContext PluginContext = FOculusAudioLibraryManager::Get().GetPluginContext();
 		if (PluginContext != nullptr)
 		{
-			CachedContext = new ovrAudioContext;
-			*CachedContext = PluginContext;
+			CachedContext = PluginContext;
 		}
 		else
 		{
@@ -332,22 +322,22 @@ ovrAudioContext UOculusAudioGeometryComponent::GetContext(UWorld* World)
 				if (AudioDevice == nullptr)
 				{
 					// This happens when cooking for native UE4 AudioMixer integration
-					return FOculusAudioContextManager::GetOrCreateSerializationContext(this);
+					CachedContext = FOculusAudioContextManager::GetOrCreateSerializationContext(this);
 				}
-
-				OculusAudioSpatializationAudioMixer* Spatializer = static_cast<OculusAudioSpatializationAudioMixer*>(AudioDevice->SpatializationPluginInterface.Get());
-				check(Spatializer != nullptr);
-				if (Spatializer->ClassID != OculusAudioSpatializationAudioMixer::MIXER_CLASS_ID)
+				else
 				{
-					UE_LOG(LogAudio, Warning, TEXT("Invalid Spatialization Plugin specified, make sure the Spatialization Plugin is set to OculusAudio and AudioMixer is enabled!"));
-					return nullptr;
+					CachedContext = FOculusAudioContextManager::GetContextForAudioDevice(AudioDevice);
+
+					if (!CachedContext)
+					{
+						CachedContext = FOculusAudioContextManager::CreateContextForAudioDevice(AudioDevice);
+					}
 				}
-				CachedContext = Spatializer->GetContext();
 			}
 		}
 	}
 
-	return (CachedContext != nullptr) ? *CachedContext : nullptr;
+	return CachedContext;
 }
 
 void UOculusAudioGeometryComponent::Serialize(FArchive & Ar)
diff --git a/Engine/Plugins/Runtime/Oculus/OculusAudio/Source/OculusAudio/Private/OculusAudioGeometryComponent.h b/Engine/Plugins/Runtime/Oculus/OculusAudio/Source/OculusAudio/Private/OculusAudioGeometryComponent.h
--- a/Engine/Plugins/Runtime/Oculus/OculusAudio/Source/OculusAudio/Private/OculusAudioGeometryComponent.h
+++ b/Engine/Plugins/Runtime/Oculus/OculusAudio/Source/OculusAudio/Private/OculusAudioGeometryComponent.h
@@ -47,6 +47,6 @@ private:
 	bool IncludeChildren = true;
 
 	ovrAudioGeometry ovrGeometry;
-	ovrAudioContext* CachedContext;
+	ovrAudioContext CachedContext;
 	FTransform PreviousTransform;
 };
\ No newline at end of file
diff --git a/Engine/Plugins/Runtime/Oculus/OculusAudio/Source/OculusAudio/Private/OculusAudioMixer.cpp b/Engine/Plugins/Runtime/Oculus/OculusAudio/Source/OculusAudio/Private/OculusAudioMixer.cpp
--- a/Engine/Plugins/Runtime/Oculus/OculusAudio/Source/OculusAudio/Private/OculusAudioMixer.cpp
+++ b/Engine/Plugins/Runtime/Oculus/OculusAudio/Source/OculusAudio/Private/OculusAudioMixer.cpp
@@ -3,6 +3,7 @@
 #include "OculusAudioMixer.h"
 #include "OculusAudioSettings.h"
 #include "OculusAudioSourceSettings.h"
+#include "OculusAudioContextManager.h"
 #include "IOculusAudioPlugin.h"
 
 
@@ -12,8 +13,7 @@ float dbToLinear(float db)
 }
 
 OculusAudioSpatializationAudioMixer::OculusAudioSpatializationAudioMixer()
-	: bOvrContextInitialized(false)
-	, Context(nullptr)
+	: Context(nullptr)
 {
 }
 
@@ -23,66 +23,33 @@ OculusAudioSpatializationAudioMixer::~OculusAudioSpatializationAudioMixer()
 	{
 		FTicker::GetCoreTicker().RemoveTicker(TickDelegateHandle);
 	}
+}
 
-	if (bOvrContextInitialized)
-	{
-		// clear context from map
-		Shutdown();
-	}
+void OculusAudioSpatializationAudioMixer::ClearContext()
+{
+	Context = nullptr;
 }
 
-void OculusAudioSpatializationAudioMixer::SetContext(ovrAudioContext* SharedContext)
+void OculusAudioSpatializationAudioMixer::Initialize(const FAudioPluginInitializationParams InitializationParams)
 {
 	FScopeLock ScopeLock(&ContextLock);
 
-	Context = SharedContext;
-
-	if (SharedContext == nullptr)
+	Context = FOculusAudioContextManager::GetContextForAudioDevice(InitializationParams.AudioDevicePtr);
+	if (!Context)
 	{
-		return;
+		Context = FOculusAudioContextManager::CreateContextForAudioDevice(InitializationParams.AudioDevicePtr);
+		check(Context);
 	}
 
 	Params.AddDefaulted(InitParams.NumSources);
 
-	ovrAudioContextConfiguration ContextConfig;
-	ContextConfig.acc_Size = sizeof(ovrAudioContextConfiguration);
-
-	// First initialize the Fast algorithm context
-	ContextConfig.acc_MaxNumSources = InitParams.NumSources;
-	ContextConfig.acc_SampleRate = InitParams.SampleRate;
-	ContextConfig.acc_BufferLength = InitParams.BufferLength;
-
-	ovrResult Result = ovrSuccess;
-	if (*Context == nullptr)
-	{
-		// Create the OVR Audio Context with a given quality
-		Result = OVRA_CALL(ovrAudio_CreateContext)(Context, &ContextConfig);
-		OVR_AUDIO_CHECK(Result, "Failed to create Oculus Audio context");
-	}
-	else
-	{
-		Result = OVRA_CALL(ovrAudio_InitializeContext)(*Context, &ContextConfig);
-		OVR_AUDIO_CHECK(Result, "Failed to initialize Oculus Audio context");
-	}
-
 	// UE4 is in centimeters
-	Result = OVRA_CALL(ovrAudio_SetUnitScale)(*Context, 0.01f);
-	OVR_AUDIO_CHECK(Result, "Failed to set unit scale");
+	ovrResult Result = OVRA_CALL(ovrAudio_SetUnitScale)(Context, 0.01f);
+	OVR_AUDIO_CHECK(Result == ovrSuccess, "Failed to set unit scale");
 
 	const UOculusAudioSettings* Settings = GetDefault<UOculusAudioSettings>();
 	ApplyOculusAudioSettings(Settings);
 
-	// Now initialize the high quality algorithm context
-	bOvrContextInitialized = true;
-}
-
-void OculusAudioSpatializationAudioMixer::Initialize(const FAudioPluginInitializationParams InitializationParams)
-{
-	if (bOvrContextInitialized)
-	{
-		return;
-	}
-
 	InitParams = InitializationParams;
 
 	TickDelegateHandle = FTicker::GetCoreTicker().AddTicker(FTickerDelegate::CreateRaw(this, &OculusAudioSpatializationAudioMixer::Tick));
@@ -91,10 +58,10 @@ void OculusAudioSpatializationAudioMixer::Initialize(const FAudioPluginInitializ
 void OculusAudioSpatializationAudioMixer::ApplyOculusAudioSettings(const UOculusAudioSettings* Settings)
 {
 	FScopeLock ScopeLock(&ContextLock);
-	ovrResult Result = OVRA_CALL(ovrAudio_Enable)(*Context, ovrAudioEnable_SimpleRoomModeling, Settings->EarlyReflections);
+	ovrResult Result = OVRA_CALL(ovrAudio_Enable)(Context, ovrAudioEnable_SimpleRoomModeling, Settings->EarlyReflections);
 	OVR_AUDIO_CHECK(Result, "Failed to enable reflections");
 
-	Result = OVRA_CALL(ovrAudio_Enable)(*Context, ovrAudioEnable_LateReverberation, Settings->LateReverberation);
+	Result = OVRA_CALL(ovrAudio_Enable)(Context, ovrAudioEnable_LateReverberation, Settings->LateReverberation);
 	OVR_AUDIO_CHECK(Result, "Failed to enable reverb");
 
 	ovrAudioBoxRoomParameters Room = { 0 };
@@ -109,27 +76,26 @@ void OculusAudioSpatializationAudioMixer::ApplyOculusAudioSettings(const UOculus
 	Room.brp_ReflectBehind = Settings->ReflectionCoefBack;
 	Room.brp_ReflectFront = Settings->ReflectionCoefFront;
 
-	Result = OVRA_CALL(ovrAudio_SetSimpleBoxRoomParameters)(*Context, &Room);
+	Result = OVRA_CALL(ovrAudio_SetSimpleBoxRoomParameters)(Context, &Room);
 	OVR_AUDIO_CHECK(Result, "Failed to set room parameters");
 
-	Result = OVRA_CALL(ovrAudio_SetSharedReverbWetLevel)(*Context, dbToLinear(Settings->ReverbWetLevel));
+	Result = OVRA_CALL(ovrAudio_SetSharedReverbWetLevel)(Context, dbToLinear(Settings->ReverbWetLevel));
 	OVR_AUDIO_CHECK(Result, "Failed to set room parameters");
 }
 
 void OculusAudioSpatializationAudioMixer::Shutdown()
 {
-	bOvrContextInitialized = false;
 }
 
 bool OculusAudioSpatializationAudioMixer::IsSpatializationEffectInitialized() const
 {
-	return bOvrContextInitialized;
+	return Context != nullptr;
 }
 
 void OculusAudioSpatializationAudioMixer::OnInitSource(const uint32 SourceId, const FName& AudioComponentUserId, USpatializationPluginSourceSettingsBase* InSettings)
 {
 	FScopeLock ScopeLock(&ContextLock);
-	check(bOvrContextInitialized);
+	check(Context != nullptr);
 
 	if (InSettings != nullptr) 
 	{
@@ -139,20 +105,20 @@ void OculusAudioSpatializationAudioMixer::OnInitSource(const uint32 SourceId, co
 		if (!Settings->EarlyReflectionsEnabled)
 			Flags |= ovrAudioSourceFlag_ReflectionsDisabled;
 
-		ovrResult Result = OVRA_CALL(ovrAudio_SetAudioSourceFlags)(*Context, SourceId, Flags);
+		ovrResult Result = OVRA_CALL(ovrAudio_SetAudioSourceFlags)(Context, SourceId, Flags);
 		OVR_AUDIO_CHECK(Result, "Failed to set audio source flags");
 
 		ovrAudioSourceAttenuationMode mode = Settings->AttenuationEnabled ? ovrAudioSourceAttenuationMode_InverseSquare : ovrAudioSourceAttenuationMode_None;
-		Result = OVRA_CALL(ovrAudio_SetAudioSourceAttenuationMode)(*Context, SourceId, mode, 1.0f);
+		Result = OVRA_CALL(ovrAudio_SetAudioSourceAttenuationMode)(Context, SourceId, mode, 1.0f);
 		OVR_AUDIO_CHECK(Result, "Failed to set audio source attenuation mode");
 
-		Result = OVRA_CALL(ovrAudio_SetAudioSourceRange)(*Context, SourceId, Settings->AttenuationRangeMinimum, Settings->AttenuationRangeMaximum);
+		Result = OVRA_CALL(ovrAudio_SetAudioSourceRange)(Context, SourceId, Settings->AttenuationRangeMinimum, Settings->AttenuationRangeMaximum);
 		OVR_AUDIO_CHECK(Result, "Failed to set audio source attenuation range");
 
-		Result = OVRA_CALL(ovrAudio_SetAudioSourceRadius)(*Context, SourceId, Settings->VolumetricRadius);
+		Result = OVRA_CALL(ovrAudio_SetAudioSourceRadius)(Context, SourceId, Settings->VolumetricRadius);
 		OVR_AUDIO_CHECK(Result, "Failed to set audio source volumetric radius");
 
-		Result = OVRA_CALL(ovrAudio_SetAudioReverbSendLevel)(*Context, SourceId, dbToLinear(Settings->ReverbSendLevel));
+		Result = OVRA_CALL(ovrAudio_SetAudioReverbSendLevel)(Context, SourceId, dbToLinear(Settings->ReverbSendLevel));
 		OVR_AUDIO_CHECK(Result, "Failed to set reverb send level");
 	}
 }
@@ -165,7 +131,7 @@ void OculusAudioSpatializationAudioMixer::SetSpatializationParameters(uint32 Voi
 void OculusAudioSpatializationAudioMixer::ProcessAudio(const FAudioPluginSourceInputData& InputData, FAudioPluginSourceOutputData& OutputData)
 {
 	FScopeLock ScopeLock(&ContextLock);
-	if (InputData.SpatializationParams && *Context)
+	if (InputData.SpatializationParams && Context)
 	{
 		Params[InputData.SourceId] = *InputData.SpatializationParams;
 
@@ -174,7 +140,7 @@ void OculusAudioSpatializationAudioMixer::ProcessAudio(const FAudioPluginSourceI
 		FVector OvrListenerForward = ToOVRVector(Params[InputData.SourceId].ListenerOrientation.GetForwardVector());
 		FVector OvrListenerUp = ToOVRVector(Params[InputData.SourceId].ListenerOrientation.GetUpVector());
 
-		ovrResult Result = OVRA_CALL(ovrAudio_SetListenerVectors)(*Context,
+		ovrResult Result = OVRA_CALL(ovrAudio_SetListenerVectors)(Context,
 			OvrListenerPosition.X, OvrListenerPosition.Y, OvrListenerPosition.Z,
 			OvrListenerForward.X, OvrListenerForward.Y, OvrListenerForward.Z,
 			OvrListenerUp.X, OvrListenerUp.Y, OvrListenerUp.Z);
@@ -184,12 +150,12 @@ void OculusAudioSpatializationAudioMixer::ProcessAudio(const FAudioPluginSourceI
 		FVector OvrPosition = ToOVRVector(Params[InputData.SourceId].EmitterWorldPosition);
 
 		// Set the source position to current audio position
-		Result = OVRA_CALL(ovrAudio_SetAudioSourcePos)(*Context, InputData.SourceId, OvrPosition.X, OvrPosition.Y, OvrPosition.Z);
+		Result = OVRA_CALL(ovrAudio_SetAudioSourcePos)(Context, InputData.SourceId, OvrPosition.X, OvrPosition.Y, OvrPosition.Z);
 		OVR_AUDIO_CHECK(Result, "Failed to set audio source position");
 
 		// Perform the processing
 		uint32 Status;
-		Result = OVRA_CALL(ovrAudio_SpatializeMonoSourceInterleaved)(*Context, InputData.SourceId, &Status, OutputData.AudioBuffer.GetData(), InputData.AudioBuffer->GetData());
+		Result = OVRA_CALL(ovrAudio_SpatializeMonoSourceInterleaved)(Context, InputData.SourceId, &Status, OutputData.AudioBuffer.GetData(), InputData.AudioBuffer->GetData());
 		OVR_AUDIO_CHECK(Result, "Failed to spatialize mono source interleaved");
 	}
 }
@@ -198,12 +164,12 @@ bool OculusAudioSpatializationAudioMixer::Tick(float DeltaTime)
 {
 	if (ContextLock.TryLock())
 	{
-		if (Context != nullptr && *Context == nullptr)
+		if (Context != nullptr)
 		{
-			ovrResult Result = OVRA_CALL(ovrAudio_UpdateRoomModel)(*Context, 1.0f);
+			ovrResult Result = OVRA_CALL(ovrAudio_UpdateRoomModel)(Context, 1.0f);
 
 			UOculusAudioSettings* settings = GetMutableDefault<UOculusAudioSettings>();
-			Result = OVRA_CALL(ovrAudio_SetPropagationQuality)(*Context, settings->PropagationQuality);
+			Result = OVRA_CALL(ovrAudio_SetPropagationQuality)(Context, settings->PropagationQuality);
 			if (Result != ovrSuccess)
 			{
 				UE_LOG(LogTemp, Warning, TEXT("Bad Propagation Quality setting %d!"), settings->PropagationQuality);
diff --git a/Engine/Plugins/Runtime/Oculus/OculusAudio/Source/OculusAudio/Private/OculusAudioMixer.h b/Engine/Plugins/Runtime/Oculus/OculusAudio/Source/OculusAudio/Private/OculusAudioMixer.h
--- a/Engine/Plugins/Runtime/Oculus/OculusAudio/Source/OculusAudio/Private/OculusAudioMixer.h
+++ b/Engine/Plugins/Runtime/Oculus/OculusAudio/Source/OculusAudio/Private/OculusAudioMixer.h
@@ -6,6 +6,7 @@
 #include "OculusAudioDllManager.h"
 #include "OculusAudioSettings.h"
 #include "Containers/Ticker.h"
+#include "AudioMixerDevice.h"
 
 /************************************************************************/
 /* OculusAudioSpatializationAudioMixer									*/
@@ -18,9 +19,7 @@ public:
 	OculusAudioSpatializationAudioMixer();
 	~OculusAudioSpatializationAudioMixer();
 
-	void SetContext(ovrAudioContext* SharedContext);
 	void ClearContext();
-	ovrAudioContext* GetContext() { return Context; }
 
 	//~ Begin IAudioSpatialization 
 	virtual void Initialize(const FAudioPluginInitializationParams InitializationParams) override;
@@ -39,18 +38,25 @@ public:
 		return FVector(InVec.Y, InVec.Z, -InVec.X);
 	}
 
+	static FVector ToOVRVector(const Audio::FChannelPositionInfo& ChannelPositionInfo)
+	{
+		FVector OvrVector;
+		OvrVector.X = ChannelPositionInfo.Radius * FMath::Sin(ChannelPositionInfo.Azimuth) * FMath::Cos(ChannelPositionInfo.Elevation);
+		OvrVector.Y = ChannelPositionInfo.Radius * FMath::Sin(ChannelPositionInfo.Azimuth) * FMath::Sin(ChannelPositionInfo.Elevation);
+		OvrVector.Z = ChannelPositionInfo.Radius * FMath::Cos(ChannelPositionInfo.Azimuth);
+
+		return OvrVector;
+	}
+
 	static const int MIXER_CLASS_ID = 0x98765432;
 	const int ClassID = MIXER_CLASS_ID;
 private:
 
 	void ApplyOculusAudioSettings(const UOculusAudioSettings* Settings);
 
-	// Whether or not the OVR audio context is initialized. We defer initialization until the first audio callback.
-	bool bOvrContextInitialized;
-
 	TArray<FSpatializationParams> Params;
 
-	ovrAudioContext* Context;
+	ovrAudioContext Context;
 	FCriticalSection ContextLock;
 
 	FAudioPluginInitializationParams InitParams;
diff --git a/Engine/Plugins/Runtime/Oculus/OculusAudio/Source/OculusAudio/Private/OculusAudioReverb.cpp b/Engine/Plugins/Runtime/Oculus/OculusAudio/Source/OculusAudio/Private/OculusAudioReverb.cpp
--- a/Engine/Plugins/Runtime/Oculus/OculusAudio/Source/OculusAudio/Private/OculusAudioReverb.cpp
+++ b/Engine/Plugins/Runtime/Oculus/OculusAudio/Source/OculusAudio/Private/OculusAudioReverb.cpp
@@ -3,6 +3,7 @@
 #include "OculusAudioReverb.h"
 #include "OculusAudioMixer.h"
 #include "OculusAudioSettings.h"
+#include "OculusAudioContextManager.h"
 
 #include "Sound/SoundSubmix.h"
 
@@ -15,15 +16,22 @@ namespace
 	}
 } // namespace <>
 
-void FSubmixEffectOculusReverbPlugin::SetContext(ovrAudioContext* SharedContext)
+void FSubmixEffectOculusReverbPlugin::ClearContext()
 {
-	FScopeLock ScopeLock(&ContextLock);
-	Context = SharedContext;
+	Context = nullptr;
 }
 
-void FSubmixEffectOculusReverbPlugin::ClearContext()
+void FSubmixEffectOculusReverbPlugin::Init(const FSoundEffectSubmixInitData& InInitData)
 {
-	Context = nullptr;
+	FAudioDevice* AudioDevice = GEngine ? GEngine->GetAudioDeviceManager()->GetAudioDeviceRaw(InInitData.DeviceID) : nullptr;
+
+	Context = FOculusAudioContextManager::GetContextForAudioDevice(AudioDevice);
+	if (!Context)
+	{
+		Context = FOculusAudioContextManager::CreateContextForAudioDevice(AudioDevice);
+	}
+
+	check(Context);
 }
 
 FSubmixEffectOculusReverbPlugin::FSubmixEffectOculusReverbPlugin()
@@ -33,32 +41,17 @@ FSubmixEffectOculusReverbPlugin::FSubmixEffectOculusReverbPlugin()
 
 void FSubmixEffectOculusReverbPlugin::OnProcessAudio(const FSoundEffectSubmixInputData& InputData, FSoundEffectSubmixOutputData& OutputData)
 {
-	FScopeLock ScopeLock(&ContextLock);
-	if (Context != nullptr && *Context != nullptr)
-	{
-		int Enabled = 0;
-		ovrResult Result = OVRA_CALL(ovrAudio_IsEnabled)(*Context, ovrAudioEnable_LateReverberation, &Enabled);
-		OVR_AUDIO_CHECK(Result, "Failed to check if reverb is Enabled");
-
-		if (Enabled != 0)
-		{
-			uint32_t Status = 0;
-			Result = OVRA_CALL(ovrAudio_MixInSharedReverbInterleaved)(*Context, &Status, OutputData.AudioBuffer->GetData());
-			OVR_AUDIO_CHECK(Result, "Failed to process reverb");
-		}
-	}
-}
+	check(Context);
 
-void OculusAudioReverb::SetContext(ovrAudioContext* SharedContext)
-{
-	if (SharedContext != nullptr)
-	{
-		Context = SharedContext;
-	}
+	int Enabled = 0;
+	ovrResult Result = OVRA_CALL(ovrAudio_IsEnabled)(Context, ovrAudioEnable_LateReverberation, &Enabled);
+	OVR_AUDIO_CHECK(Result, "Failed to check if reverb is Enabled");
 
-	if (SubmixEffect.IsValid())
+	if (Enabled != 0)
 	{
-		CastEffectToPluginSharedPtr(SubmixEffect)->SetContext(SharedContext);
+		uint32_t Status = 0;
+		Result = OVRA_CALL(ovrAudio_MixInSharedReverbInterleaved)(Context, &Status, OutputData.AudioBuffer->GetData());
+		OVR_AUDIO_CHECK(Result, "Failed to process reverb");
 	}
 }
 
diff --git a/Engine/Plugins/Runtime/Oculus/OculusAudio/Source/OculusAudio/Private/OculusAudioReverb.h b/Engine/Plugins/Runtime/Oculus/OculusAudio/Source/OculusAudio/Private/OculusAudioReverb.h
--- a/Engine/Plugins/Runtime/Oculus/OculusAudio/Source/OculusAudio/Private/OculusAudioReverb.h
+++ b/Engine/Plugins/Runtime/Oculus/OculusAudio/Source/OculusAudio/Private/OculusAudioReverb.h
@@ -13,13 +13,9 @@ class FSubmixEffectOculusReverbPlugin : public FSoundEffectSubmix
 public:
 	FSubmixEffectOculusReverbPlugin();
 
-	void SetContext(ovrAudioContext* SharedContext);
 	void ClearContext();
 
-	virtual void Init(const FSoundEffectSubmixInitData& InSampleRate) override
-	{
-		return; // PAS
-	}
+	virtual void Init(const FSoundEffectSubmixInitData& InInitData) override;
 	virtual uint32 GetDesiredInputChannelCountOverride() const override
 	{
 		static const int STEREO = 2;
@@ -31,7 +27,7 @@ public:
 		return; // PAS
 	}
 private:
-	ovrAudioContext* Context;
+	ovrAudioContext Context;
 	FCriticalSection ContextLock;
 };
 
@@ -49,7 +45,6 @@ public:
 		// empty
 	}
 
-	void SetContext(ovrAudioContext* SharedContext);
 	void ClearContext();
 
 	virtual void OnInitSource(const uint32 SourceId, const FName& AudioComponentUserId, const uint32 NumChannels, UReverbPluginSourceSettingsBase* InSettings) override
diff --git a/Engine/Plugins/Runtime/Oculus/OculusAudio/Source/OculusAudio/Public/OculusAmbisonicsSettings.h b/Engine/Plugins/Runtime/Oculus/OculusAudio/Source/OculusAudio/Public/OculusAmbisonicsSettings.h
--- a/Engine/Plugins/Runtime/Oculus/OculusAudio/Source/OculusAudio/Public/OculusAmbisonicsSettings.h
+++ b/Engine/Plugins/Runtime/Oculus/OculusAudio/Source/OculusAudio/Public/OculusAmbisonicsSettings.h
@@ -8,7 +8,7 @@
 
 
 UENUM()
-enum class EAmbisonicMode : uint8
+enum class EOculusAudioAmbisonicsMode : uint8
 {
 	// High quality ambisonic spatialization method
 	SphericalHarmonics,
@@ -17,16 +17,6 @@ enum class EAmbisonicMode : uint8
 	VirtualSpeakers,
 };
 
-UENUM()
-enum class EAmbisonicFormat : uint8
-{
-	// Standard B-Format, WXYZ
-	FuMa UMETA(DisplayName = "FuMa (B-Format, WXYZ)"),
-
-	// ACN/SN3D Standard, WYZX
-	AmbiX UMETA(DisplayName = "AmbiX (ACN/SN3D, WYZX)"),
-};
-
 USTRUCT(BlueprintType)
 struct OCULUSAUDIO_API FSubmixEffectOculusAmbisonicSpatializerSettings
 {
@@ -34,25 +24,49 @@ struct OCULUSAUDIO_API FSubmixEffectOculusAmbisonicSpatializerSettings
 
 		// Ambisonic spatialization mode
 	UPROPERTY(EditAnywhere, BlueprintReadWrite, Category = Realtime)
-	EAmbisonicMode AmbisonicMode;
+	EOculusAudioAmbisonicsMode AmbisonicMode;
 
 	FSubmixEffectOculusAmbisonicSpatializerSettings()
-		: AmbisonicMode(EAmbisonicMode::SphericalHarmonics)
+		: AmbisonicMode(EOculusAudioAmbisonicsMode::SphericalHarmonics)
 	{
 	}
 };
 
+class FOculusAudioSoundfieldSettingsProxy : public ISoundfieldEncodingSettingsProxy
+{
+
+public:
+	const EOculusAudioAmbisonicsMode SpatailizationMode;
+
+	FOculusAudioSoundfieldSettingsProxy(EOculusAudioAmbisonicsMode InSpatializationMode)
+		: SpatailizationMode(InSpatializationMode)
+	{}
+
+	virtual uint32 GetUniqueId() const override
+	{
+		return (uint32)SpatailizationMode;
+	}
+
+
+	virtual TUniquePtr<ISoundfieldEncodingSettingsProxy> Duplicate() const override
+	{
+		return TUniquePtr<ISoundfieldEncodingSettingsProxy>(new FOculusAudioSoundfieldSettingsProxy(SpatailizationMode));
+	}
+
+};
+
 UCLASS()
-class OCULUSAUDIO_API UOculusAmbisonicsSettings : public UAmbisonicsSubmixSettingsBase
+class OCULUSAUDIO_API UOculusAudioSoundfieldSettings : public USoundfieldEncodingSettingsBase
 {
 	GENERATED_BODY()
 
 public:
 	UPROPERTY(EditAnywhere, BlueprintReadOnly, Category = Ambisonics)
-	EAmbisonicMode SpatializationMode;
-
-	UPROPERTY(EditAnywhere, BlueprintReadOnly, Category = Ambisonics)
-	EAmbisonicFormat ChannelOrder;
+	EOculusAudioAmbisonicsMode SpatializationMode;
 
+	virtual TUniquePtr<ISoundfieldEncodingSettingsProxy> GetProxy() const override
+	{
+		return TUniquePtr<ISoundfieldEncodingSettingsProxy>(new FOculusAudioSoundfieldSettingsProxy(SpatializationMode));
+	}
 };
 
diff --git a/Engine/Plugins/Runtime/Oculus/OculusAudio/Source/OculusAudioEditor/Private/OculusAmbisonicSettingsFactory.cpp b/Engine/Plugins/Runtime/Oculus/OculusAudio/Source/OculusAudioEditor/Private/OculusAmbisonicSettingsFactory.cpp
--- a/Engine/Plugins/Runtime/Oculus/OculusAudio/Source/OculusAudioEditor/Private/OculusAmbisonicSettingsFactory.cpp
+++ b/Engine/Plugins/Runtime/Oculus/OculusAudio/Source/OculusAudioEditor/Private/OculusAmbisonicSettingsFactory.cpp
@@ -30,7 +30,7 @@ const TArray<FText>& FAssetTypeActions_OculusAmbisonicsSettings::GetSubMenus() c
 
 UClass* FAssetTypeActions_OculusAmbisonicsSettings::GetSupportedClass() const
 {
-    return UOculusAmbisonicsSettings::StaticClass();
+    return UOculusAudioSoundfieldSettings::StaticClass();
 }
 
 uint32 FAssetTypeActions_OculusAmbisonicsSettings::GetCategories()
@@ -41,7 +41,7 @@ uint32 FAssetTypeActions_OculusAmbisonicsSettings::GetCategories()
 UOculusAmbisonicsSettingsFactory::UOculusAmbisonicsSettingsFactory(const FObjectInitializer& ObjectInitializer)
 	: Super(ObjectInitializer)
 {
-	SupportedClass = UOculusAmbisonicsSettings::StaticClass();
+	SupportedClass = UOculusAudioSoundfieldSettings::StaticClass();
 
 	bCreateNew = true;
 	bEditorImport = true;
@@ -51,7 +51,7 @@ UOculusAmbisonicsSettingsFactory::UOculusAmbisonicsSettingsFactory(const FObject
 UObject* UOculusAmbisonicsSettingsFactory::FactoryCreateNew(UClass* Class, UObject* InParent, FName Name, EObjectFlags Flags, UObject* Context,
     FFeedbackContext* Warn)
 {
-    return NewObject<UOculusAmbisonicsSettings>(InParent, Name, Flags);
+    return NewObject<UOculusAudioSoundfieldSettings>(InParent, Name, Flags);
 }
 
 uint32 UOculusAmbisonicsSettingsFactory::GetMenuCategories() const
diff --git a/Engine/Plugins/Runtime/ResonanceAudio/Source/ResonanceAudio/Private/ResonanceAudioAmbisonics.cpp b/Engine/Plugins/Runtime/ResonanceAudio/Source/ResonanceAudio/Private/ResonanceAudioAmbisonics.cpp
--- a/Engine/Plugins/Runtime/ResonanceAudio/Source/ResonanceAudio/Private/ResonanceAudioAmbisonics.cpp
+++ b/Engine/Plugins/Runtime/ResonanceAudio/Source/ResonanceAudio/Private/ResonanceAudioAmbisonics.cpp
@@ -3,162 +3,394 @@
 //
 #include "ResonanceAudioAmbisonics.h"
 #include "ResonanceAudioAmbisonicsSettings.h"
+#include "ResonanceAudioPluginListener.h"
+#include "SoundFieldRendering.h"
+
+static int32 UseReverbPluginForResonanceOutputCVar = 1;
+FAutoConsoleVariableRef CVarUseReverbPluginForResonanceOutput(
+	TEXT("au.resonance.UseReverbPluginForOutput"),
+	UseReverbPluginForResonanceOutputCVar,
+	TEXT("If set to 0, will allow for Resonance to render it's output directly from a submix even if it is not set as the spatialization or reverb plugin.\n")
+	TEXT("0: Disable, >0: Enable"),
+	ECVF_Default);
 
 namespace ResonanceAudio
 {
+	class FEncoder : public ISoundfieldEncoderStream
+	{
+	private:
+		vraudio::ResonanceAudioApi* ResonanceSystem;
+		bool bShouldDestroyResonanceSystem;
 
-	const int32 DefaultNumChannels = 4;
+		TArray<vraudio::SourceId> SourceIdArray;
+		TArray<Audio::AlignedFloatBuffer> DeinterleavedBuffers;
+		int32 NumChannels;
+
+		void ChangeNumChannels(int32 NewNumChannels, vraudio::RenderingMode InRenderingMode)
+		{
+			check(ResonanceSystem);
+			for (vraudio::SourceId Source : SourceIdArray)
+			{
+				ResonanceSystem->DestroySource(Source);
+			}
 
+			DeinterleavedBuffers.Reset();
 
-	int32 GetNumChannelsForAmbisonicsOrder(EAmbisonicsOrder InOrder)
-	{
-		switch (InOrder)
+			if (NewNumChannels > 1)
+			{
+				for (int32 Index = 0; Index < NewNumChannels; Index++)
+				{
+					DeinterleavedBuffers.AddDefaulted();
+					SourceIdArray.Add(ResonanceSystem->CreateSoundObjectSource(InRenderingMode));
+				}
+			}
+
+			NumChannels = NewNumChannels;
+		}
+
+	public:
+		FEncoder(vraudio::ResonanceAudioApi* InResonanceSystem, bool bInShouldDestroyResonanceAPI, int32 InNumChannels, vraudio::RenderingMode InRenderingMode)
+			: ResonanceSystem(InResonanceSystem)
+			, bShouldDestroyResonanceSystem(bInShouldDestroyResonanceAPI)
+			, NumChannels(InNumChannels)
 		{
-			case EAmbisonicsOrder::SecondOrder:
+			check(ResonanceSystem);
+			check(NumChannels > 0);
+			
+			for (int32 Index = 0; Index < NumChannels; Index++)
 			{
-				return 9;
+				DeinterleavedBuffers.AddDefaulted();
+				SourceIdArray.Add(ResonanceSystem->CreateSoundObjectSource(InRenderingMode));
 			}
+		}
 
-			case EAmbisonicsOrder::ThirdOrder:
+		virtual ~FEncoder()
+		{
+			check(ResonanceSystem);
+			for (vraudio::SourceId Source : SourceIdArray)
 			{
-				return 16;
+				ResonanceSystem->DestroySource(Source);
 			}
 
-			case EAmbisonicsOrder::FirstOrder:
-			default:
+			if (bShouldDestroyResonanceSystem)
 			{
-				return 4;
+				FResonanceAudioPluginListener::RemoveResonanceAPIForAudioDevice(ResonanceSystem);
+				delete ResonanceSystem;
+				ResonanceSystem = nullptr;
 			}
 		}
-	}
 
-	FResonanceAudioAmbisonicsMixer::FResonanceAudioAmbisonicsMixer()
-		: ResonanceAudioApi(nullptr)
-	{
-	}
 
-	void FResonanceAudioAmbisonicsMixer::Shutdown()
+		virtual void Encode(const FSoundfieldEncoderInputData& InputData, ISoundfieldAudioPacket& OutputData) override
+		{
+			OutputData.Reset();
+			EncodeAndMixIn(InputData, OutputData);
+		}
+
+
+		virtual void EncodeAndMixIn(const FSoundfieldEncoderInputData& InputData, ISoundfieldAudioPacket& OutputData) override
+		{
+			check(ResonanceSystem);
+			FResonancePacket& OutPacket = DowncastSoundfieldRef<FResonancePacket>(OutputData);
+			OutPacket.SetResonanceApi(ResonanceSystem);
+
+			const FResonanceAmbisonicsSettingsProxy& Settings = DowncastSoundfieldRef<const FResonanceAmbisonicsSettingsProxy>(InputData.InputSettings);
+
+			if (NumChannels == 1)
+			{
+				check(SourceIdArray.Num() == 1);
+
+				// Don't use the DeinterleavedBuffers, submit the audio directly.
+				ResonanceSystem->SetInterleavedBuffer(SourceIdArray[0], InputData.AudioBuffer.GetData(), 1, InputData.AudioBuffer.Num());
+			}
+			else
+			{
+				check(InputData.PositionalData.ChannelPositions);
+				const TArray<Audio::FChannelPositionInfo>& ChannelPositions = *InputData.PositionalData.ChannelPositions;
+
+
+				if (ChannelPositions.Num() != NumChannels)
+				{
+					ChangeNumChannels(ChannelPositions.Num(), Settings.RenderingMode);
+				}
+
+				const int32 NumFrames = InputData.AudioBuffer.Num() / InputData.NumChannels;
+				check(DeinterleavedBuffers.Num() == InputData.NumChannels);
+
+				for (int32 ChannelIndex = 0; ChannelIndex < InputData.NumChannels; ChannelIndex++)
+				{
+					// Deinterleave the audio for this channel.
+					Audio::AlignedFloatBuffer& DeinterleavedBuffer = DeinterleavedBuffers[ChannelIndex];
+					DeinterleavedBuffer.Reset();
+					DeinterleavedBuffer.AddUninitialized(NumFrames);
+
+					for (int32 FrameIndex = 0; FrameIndex < NumFrames; FrameIndex++)
+					{
+						DeinterleavedBuffer[FrameIndex] = InputData.AudioBuffer[FrameIndex * NumChannels + ChannelIndex];
+					}
+
+					// Set the position and audio buffer for this source.
+					vraudio::SourceId& SourceId = SourceIdArray[ChannelIndex];
+					FVector ResonancePosition = ConvertToResonanceAudioCoordinates(ChannelPositions[ChannelIndex]);
+					ResonanceSystem->SetSourcePosition(SourceId, ResonancePosition.X, ResonancePosition.Y, ResonancePosition.Z);
+					ResonanceSystem->SetInterleavedBuffer(SourceId, DeinterleavedBuffer.GetData(), 1, NumFrames);
+				}
+			}
+		}
+	};
+
+	class FDecoder : public ISoundfieldDecoderStream
 	{
-	}
+	public:
+		virtual void Decode(const FSoundfieldDecoderInputData& InputData, FSoundfieldDecoderOutputData& OutputData) override
+		{
+			DecodeAndMixIn(InputData, OutputData);
+		}
+
+
+		virtual void DecodeAndMixIn(const FSoundfieldDecoderInputData& InputData, FSoundfieldDecoderOutputData& OutputData) override
+		{
+			// If the reverb plugin is enabled, FillInterleavedOutputBuffer will be called there.
+			// TODO: detect whether the reverb plugin is enabled on this audio device to toggle this behavior automatically.
+			if (!UseReverbPluginForResonanceOutputCVar)
+			{
+				const int32 NumChannels = InputData.PositionalData.NumChannels;
+				const FResonancePacket& Packet = DowncastSoundfieldRef<const FResonancePacket>(InputData.SoundfieldBuffer);
+				OutputData.AudioBuffer.Reset();
+				OutputData.AudioBuffer.AddZeroed(NumChannels * InputData.NumFrames);
+
+				if (vraudio::ResonanceAudioApi* ResonanceSystem = const_cast<vraudio::ResonanceAudioApi*>(Packet.GetResonanceApi()))
+				{
+					ResonanceSystem->FillInterleavedOutputBuffer(NumChannels, InputData.NumFrames, OutputData.AudioBuffer.GetData());
+				}
+			}
+		}
 
-	int32 FResonanceAudioAmbisonicsMixer::GetNumChannelsForAmbisonicsFormat(UAmbisonicsSubmixSettingsBase* InSettings)
+	};
+
+	/**
+	 * This class takes an audio stream in Unreal's internal ambisonics format and mixes it into Resonance's ambisonics bed.
+	 */
+	class FAmbisonicsTranscoder : public ISoundfieldTranscodeStream
 	{
-		UResonanceAudioAmbisonicsSettings* CastedSettings = Cast<UResonanceAudioAmbisonicsSettings>(InSettings);
+	private:
+		vraudio::ResonanceAudioApi* ResonanceSystem;
+		bool bShouldDestroyResonanceSystem;
+		vraudio::SourceId AmbisonicsSourceId;
+		int32 NumChannels;
 
-		if (CastedSettings != nullptr)
+		void ChangeNumChannels(int32 NewNumChannels)
 		{
-			const EAmbisonicsOrder Order = CastedSettings->AmbisonicsOrder;
-			return GetNumChannelsForAmbisonicsOrder(Order);
+			check(ResonanceSystem);
+			ResonanceSystem->DestroySource(AmbisonicsSourceId);
+
+			NumChannels = NewNumChannels;
+			AmbisonicsSourceId = ResonanceSystem->CreateAmbisonicSource(NumChannels);
 		}
-		else
+
+	public:
+		FAmbisonicsTranscoder(vraudio::ResonanceAudioApi* InResonanceSystem, int32 InNumChannels, bool bInShouldDestroyResonanceSystem)
+			: ResonanceSystem(InResonanceSystem)
+			, bShouldDestroyResonanceSystem(bInShouldDestroyResonanceSystem)
+			, NumChannels(InNumChannels)
 		{
-			return 4;
+			check(ResonanceSystem);
+			AmbisonicsSourceId = ResonanceSystem->CreateAmbisonicSource(NumChannels);
 		}
-	}
 
-	void FResonanceAudioAmbisonicsMixer::OnOpenEncodingStream(const uint32 StreamId, UAmbisonicsSubmixSettingsBase* InSettings)
+		virtual ~FAmbisonicsTranscoder()
+		{
+			check(ResonanceSystem);
+			ResonanceSystem->DestroySource(AmbisonicsSourceId);
+
+			if (bShouldDestroyResonanceSystem)
+			{
+				FResonanceAudioPluginListener::RemoveResonanceAPIForAudioDevice(ResonanceSystem);
+				delete ResonanceSystem;
+				ResonanceSystem = nullptr;
+			}
+		}
+
+		virtual void Transcode(const ISoundfieldAudioPacket& InputData, const ISoundfieldEncodingSettingsProxy& InputSettings, ISoundfieldAudioPacket& OutputData, const ISoundfieldEncodingSettingsProxy& OutputSettings) override
+		{
+			TranscodeAndMixIn(InputData, InputSettings, OutputData, OutputSettings);
+		}
+
+
+		virtual void TranscodeAndMixIn(const ISoundfieldAudioPacket& InputData, const ISoundfieldEncodingSettingsProxy& InputSettings, ISoundfieldAudioPacket& PacketToSumTo, const ISoundfieldEncodingSettingsProxy& OutputSettings) override
+		{
+			const FAmbisonicsSoundfieldBuffer& AmbisonicsBuffer = DowncastSoundfieldRef<const FAmbisonicsSoundfieldBuffer>(InputData);
+			FResonancePacket& OutPacket = DowncastSoundfieldRef<FResonancePacket>(PacketToSumTo);
+			
+			check(ResonanceSystem);
+
+			if (AmbisonicsBuffer.NumChannels != NumChannels)
+			{
+				ChangeNumChannels(AmbisonicsBuffer.NumChannels);
+			}
+
+			ResonanceSystem->SetInterleavedBuffer(AmbisonicsSourceId, AmbisonicsBuffer.AudioBuffer.GetData(), NumChannels, AmbisonicsBuffer.AudioBuffer.Num() / NumChannels);
+
+			OutPacket.SetResonanceApi(ResonanceSystem);
+		}
+
+	};
+
+	class FResonanceMixer : public ISoundfieldMixerStream
+	{
+	public:
+
+		virtual void MixTogether(const FSoundfieldMixerInputData& InputData, ISoundfieldAudioPacket& PacketToSumTo) override
+		{
+			// All this does is forward the resonance system to the PacketToSumTo.
+			const FResonancePacket& InPacket = DowncastSoundfieldRef<const FResonancePacket>(InputData.InputPacket);
+			FResonancePacket& OutPacket = DowncastSoundfieldRef<FResonancePacket>(PacketToSumTo);
+
+			OutPacket.SetResonanceApi(InPacket.GetResonanceApi());
+		}
+
+	};
+
+	FAmbisonicsFactory::FAmbisonicsFactory()
 	{
-		UE_LOG(LogResonanceAudio, Warning, TEXT("Encoding into an ambisonics is not currently supported by the Resonance Audio plugin."));
+		ISoundfieldFactory::RegisterSoundfieldFormat(this);
 	}
 
-	void FResonanceAudioAmbisonicsMixer::OnCloseEncodingStream(const uint32 SourceId)
+	FAmbisonicsFactory::~FAmbisonicsFactory()
 	{
-		//TODO: Implement.
+		ISoundfieldFactory::UnregisterSoundfieldFormat(this);
 	}
 
-	void FResonanceAudioAmbisonicsMixer::EncodeToAmbisonics(const uint32 SourceId, const FAmbisonicsEncoderInputData& InputData, FAmbisonicsEncoderOutputData& OutputData, UAmbisonicsSubmixSettingsBase* InParams)
+	FName FAmbisonicsFactory::GetSoundfieldFormatName()
 	{
-		//TODO: Implement.
+		static FName ResonanceSoundfieldFormat = TEXT("Resonance Binaural Spatialization");
+		return ResonanceSoundfieldFormat;
 	}
 
-	void FResonanceAudioAmbisonicsMixer::OnOpenDecodingStream(const uint32 StreamId, UAmbisonicsSubmixSettingsBase* InSettings, FAmbisonicsDecoderPositionalData& SpecifiedOutputPositions)
+	TUniquePtr<ISoundfieldEncoderStream> FAmbisonicsFactory::CreateEncoderStream(const FAudioPluginInitializationParams& InitInfo, const ISoundfieldEncodingSettingsProxy& InitialSettings)
 	{
-		if (ResonanceAudioApi == nullptr)
-		{
-			return;
-		}
+		const FResonanceAmbisonicsSettingsProxy& Settings = DowncastSoundfieldRef<const FResonanceAmbisonicsSettingsProxy>(InitialSettings);
+		FAudioDevice* AudioDevice = InitInfo.AudioDevicePtr;
 
-		int32 NumChannels = 4;
+		check(AudioDevice);
 
-		UResonanceAudioAmbisonicsSettings* MyAmbisonicsSettings = Cast<UResonanceAudioAmbisonicsSettings>(InSettings);
-		if(MyAmbisonicsSettings != nullptr)
+		vraudio::ResonanceAudioApi* ResonanceApi = FResonanceAudioPluginListener::GetResonanceAPIForAudioDevice(AudioDevice);
+
+		// If there's no globablly available Resonance api, spawn a new one.
+		if (!ResonanceApi)
 		{
-			NumChannels = GetNumChannelsForAmbisonicsOrder(MyAmbisonicsSettings->AmbisonicsOrder);
+			// TODO: Currently, we don't delete the resonance API we create here until the audio device is destroyed.
+			// If we're going to handle things like dynamic submix routing, the Resonance API lifecycle needs to be reworked to be more explicit.
+			ResonanceApi = CreateNewResonanceApiInstance(AudioDevice, InitInfo);
+			return TUniquePtr<ISoundfieldEncoderStream>(new FEncoder(ResonanceApi, false, InitInfo.NumOutputChannels, Settings.RenderingMode));
 		}
+		else
+		{
+			return TUniquePtr<ISoundfieldEncoderStream>(new FEncoder(ResonanceApi, false, InitInfo.NumOutputChannels, Settings.RenderingMode));
+		}
+	}
 
-		ResonanceSourceIdMap.Add(StreamId, ResonanceAudioApi->CreateAmbisonicSource(NumChannels));
+	TUniquePtr<ISoundfieldDecoderStream> FAmbisonicsFactory::CreateDecoderStream(const FAudioPluginInitializationParams& InitInfo, const ISoundfieldEncodingSettingsProxy& InitialSettings)
+	{
+		return TUniquePtr<ISoundfieldDecoderStream>(new FDecoder());
 	}
 
-	void FResonanceAudioAmbisonicsMixer::OnCloseDecodingStream(const uint32 StreamId)
+	TUniquePtr<ISoundfieldTranscodeStream> FAmbisonicsFactory::CreateTranscoderStream(const FName SourceFormat, const ISoundfieldEncodingSettingsProxy& InitialSourceSettings, const FName DestinationFormat, const ISoundfieldEncodingSettingsProxy& InitialDestinationSettings, const FAudioPluginInitializationParams& InitInfo)
 	{
-		if (ResonanceAudioApi == nullptr)
-		{
-			return;
-		}
+		check(SourceFormat == GetUnrealAmbisonicsFormatName());
+		check(DestinationFormat == GetSoundfieldFormatName());
+
+		const FAmbisonicsSoundfieldSettings& SourceSettings = DowncastSoundfieldRef<const FAmbisonicsSoundfieldSettings>(InitialSourceSettings);
+		const FResonanceAmbisonicsSettingsProxy& DestinationSettings = DowncastSoundfieldRef<const FResonanceAmbisonicsSettingsProxy>(InitialDestinationSettings);
 
-		RaSourceId* InStreamId = ResonanceSourceIdMap.Find(StreamId);
-		if (InStreamId != nullptr)
+		int32 NumChannels = (SourceSettings.Order + 1) * (SourceSettings.Order + 1);
+
+		FAudioDevice* AudioDevice = InitInfo.AudioDevicePtr;
+		
+		vraudio::ResonanceAudioApi* ResonanceApi = FResonanceAudioPluginListener::GetResonanceAPIForAudioDevice(AudioDevice);
+
+		// If there's no globabally available Resonance api, spawn a new one.
+		if (!ResonanceApi)
 		{
-			ResonanceAudioApi->DestroySource(*InStreamId);
-			ResonanceSourceIdMap.Remove(*InStreamId);
+			ResonanceApi = CreateNewResonanceApiInstance(AudioDevice, InitInfo);
+			return TUniquePtr<ISoundfieldTranscodeStream>(new FAmbisonicsTranscoder(ResonanceApi, NumChannels, true));
 		}
 		else
 		{
-			UE_LOG(LogResonanceAudio, Warning, TEXT("Tried to close a stream id (%d) we did not open."), StreamId);
+			return TUniquePtr<ISoundfieldTranscodeStream>(new FAmbisonicsTranscoder(ResonanceApi, NumChannels, false));
 		}
 	}
 
-	void FResonanceAudioAmbisonicsMixer::DecodeFromAmbisonics(const uint32 StreamId, const FAmbisonicsDecoderInputData& InputData, FAmbisonicsDecoderPositionalData& SpecifiedOutputPositions, FAmbisonicsDecoderOutputData& OutputData)
+	TUniquePtr<ISoundfieldMixerStream> FAmbisonicsFactory::CreateMixerStream(const ISoundfieldEncodingSettingsProxy& InitialSettings)
 	{
-		if (ResonanceAudioApi == nullptr)
-		{
-			return;
-		}
+		return TUniquePtr<ISoundfieldMixerStream>(new FResonanceMixer());
+	}
 
-		
-		//Get the source ID for this Ambisonics file.
-		RaSourceId* ThisStream = ResonanceSourceIdMap.Find(StreamId);
+	TUniquePtr<ISoundfieldAudioPacket> FAmbisonicsFactory::CreateEmptyPacket()
+	{
+		return TUniquePtr<ISoundfieldAudioPacket>(new FResonancePacket(nullptr));
+	}
+
+	bool FAmbisonicsFactory::IsTranscodeRequiredBetweenSettings(const ISoundfieldEncodingSettingsProxy& SourceSettings, const ISoundfieldEncodingSettingsProxy& DestinationSettings)
+	{
+		// Since all audio is encoded to third order ambisonics upon encoding for Resonance, we don't need to transcode anything.
+		return false;
+	}
 
-		//If it doesn't exist already, create a new source.
-		if (ThisStream == nullptr)
+	bool FAmbisonicsFactory::CanTranscodeFromSoundfieldFormat(FName SourceFormat, const ISoundfieldEncodingSettingsProxy& SourceEncodingSettings)
+	{
+		FName AmbisonicsSoundfieldFormatName = GetUnrealAmbisonicsFormatName();
+		if (SourceFormat == AmbisonicsSoundfieldFormatName)
 		{
-			ThisStream = &ResonanceSourceIdMap.Add(StreamId, ResonanceAudioApi->CreateAmbisonicSource(InputData.NumChannels));
+			return true;
 		}
-			
-		int32 NumFrames = InputData.AudioBuffer->Num() / InputData.NumChannels;
-		float* InputAudio = InputData.AudioBuffer->GetData();
+		else
+		{
+			return false;
+		}
+	}
 
-		ResonanceAudioApi->SetInterleavedBuffer(*ThisStream, InputAudio, InputData.NumChannels, NumFrames);
-		
+	bool FAmbisonicsFactory::CanTranscodeToSoundfieldFormat(FName DestinationFormat, const ISoundfieldEncodingSettingsProxy& DestinationEncodingSettings)
+	{
+		return false;
 	}
 
-	bool FResonanceAudioAmbisonicsMixer::ShouldReencodeBetween(UAmbisonicsSubmixSettingsBase* SourceSubmixSettings, UAmbisonicsSubmixSettingsBase* DestinationSubmixSettings)
+	UClass* FAmbisonicsFactory::GetCustomEncodingSettingsClass() const
 	{
-		return true;
+		return UResonanceAudioSoundfieldSettings::StaticClass();
 	}
 
-	void FResonanceAudioAmbisonicsMixer::Initialize(const FAudioPluginInitializationParams InitializationParams)
+	const USoundfieldEncodingSettingsBase* FAmbisonicsFactory::GetDefaultEncodingSettings()
 	{
-		// Nothing to do here for now.
+		return GetDefault<UResonanceAudioSoundfieldSettings>();
 	}
 
-	UClass* FResonanceAudioAmbisonicsMixer::GetCustomSettingsClass()
+	vraudio::ResonanceAudioApi* FAmbisonicsFactory::CreateNewResonanceApiInstance(FAudioDevice* AudioDevice, const FAudioPluginInitializationParams& InInitInfo)
 	{
-		return UResonanceAudioAmbisonicsSettings::StaticClass();
+		const size_t FramesPerBuffer = static_cast<size_t>(AudioDevice->GetBufferLength());
+		const int SampleRate = static_cast<int>(AudioDevice->GetSampleRate());
+
+		vraudio::ResonanceAudioApi* ResonanceApi = CreateResonanceAudioApi(FResonanceAudioModule::GetResonanceAudioDynamicLibraryHandle(), 2 /* num channels */, FramesPerBuffer, SampleRate);
+		FResonanceAudioPluginListener::SetResonanceAPIForAudioDevice(AudioDevice, ResonanceApi);
+
+		return ResonanceApi;
 	}
 
-	UAmbisonicsSubmixSettingsBase* FResonanceAudioAmbisonicsMixer::GetDefaultSettings()
+	void FResonancePacket::Serialize(FArchive& Ar)
 	{
-		static UResonanceAudioAmbisonicsSettings* DefaultAmbisonicsSettings = nullptr;
-		
-		if (DefaultAmbisonicsSettings == nullptr)
-		{
-			DefaultAmbisonicsSettings = NewObject<UResonanceAudioAmbisonicsSettings>();
-			DefaultAmbisonicsSettings->AmbisonicsOrder = EAmbisonicsOrder::FirstOrder;
-			DefaultAmbisonicsSettings->AddToRoot();
-		}
+		// Since the FResonancePacket is just a handle to the API, 
+		// we don't have anything to serialize.
+	}
 
-		return DefaultAmbisonicsSettings;
+	TUniquePtr<ISoundfieldAudioPacket> FResonancePacket::Duplicate() const
+	{
+		FResonancePacket* Packet = new FResonancePacket(ResonanceSystem);
+		return TUniquePtr<ISoundfieldAudioPacket>(Packet);
+	}
+
+	void FResonancePacket::Reset()
+	{
+		ResonanceSystem = nullptr;
 	}
 
 }
\ No newline at end of file
diff --git a/Engine/Plugins/Runtime/ResonanceAudio/Source/ResonanceAudio/Private/ResonanceAudioAmbisonics.h b/Engine/Plugins/Runtime/ResonanceAudio/Source/ResonanceAudio/Private/ResonanceAudioAmbisonics.h
--- a/Engine/Plugins/Runtime/ResonanceAudio/Source/ResonanceAudio/Private/ResonanceAudioAmbisonics.h
+++ b/Engine/Plugins/Runtime/ResonanceAudio/Source/ResonanceAudio/Private/ResonanceAudioAmbisonics.h
@@ -8,34 +8,73 @@
 
 namespace ResonanceAudio
 {
-	class FResonanceAudioAmbisonicsMixer : public IAmbisonicsMixer
+	// Packet passed between Resonance submixes. All it holds onto is a raw pointer
+	// to the Resonance system used to encode whatever sources were passed in.
+	class FResonancePacket : public ISoundfieldAudioPacket
 	{
+	private:
+		FResonancePacket() {}
+
+		vraudio::ResonanceAudioApi* ResonanceSystem;
 	public:
-		FResonanceAudioAmbisonicsMixer();
 
-		// Begin IAmbisonicsMixer
-		virtual void Shutdown() override;
-		virtual int32 GetNumChannelsForAmbisonicsFormat(UAmbisonicsSubmixSettingsBase* InSettings) override;
-		virtual void OnOpenEncodingStream(const uint32 StreamId, UAmbisonicsSubmixSettingsBase* InSettings) override;
-		virtual void OnCloseEncodingStream(const uint32 StreamId) override;
-		virtual void EncodeToAmbisonics(const uint32 StreamId, const FAmbisonicsEncoderInputData& InputData, FAmbisonicsEncoderOutputData& OutputData, UAmbisonicsSubmixSettingsBase* InParams) override;
-		virtual void OnOpenDecodingStream(const uint32 StreamId, UAmbisonicsSubmixSettingsBase* InParams, FAmbisonicsDecoderPositionalData& SpecifiedOutputPositions) override;
-		virtual void OnCloseDecodingStream(const uint32 StreamId) override;
-		virtual void DecodeFromAmbisonics(const uint32 StreamId, const FAmbisonicsDecoderInputData& InputData, FAmbisonicsDecoderPositionalData& SpecifiedOutputPositions, FAmbisonicsDecoderOutputData& OutputData) override;
-		virtual bool ShouldReencodeBetween(UAmbisonicsSubmixSettingsBase* SourceSubmixSettings, UAmbisonicsSubmixSettingsBase* DestinationSubmixSettings) override;
-		virtual void Initialize(const FAudioPluginInitializationParams InitializationParams) override;
-		virtual UClass* GetCustomSettingsClass() override;
-		virtual UAmbisonicsSubmixSettingsBase* GetDefaultSettings() override;
-		//~ IAmbisonicsMixer
+		FResonancePacket(vraudio::ResonanceAudioApi* InSystem)
+			: ResonanceSystem(InSystem)
+		{
+		}
+
+		virtual void Serialize(FArchive& Ar) override;
+		virtual TUniquePtr<ISoundfieldAudioPacket> Duplicate() const override;
+		virtual void Reset() override;
+
+		vraudio::ResonanceAudioApi* GetResonanceApi() const
+		{
+			return ResonanceSystem;
+		}
+
+		void SetResonanceApi(vraudio::ResonanceAudioApi* InResonanceSystem)
+		{
+			check(InResonanceSystem);
+			ResonanceSystem = InResonanceSystem;
+		}
+
+		vraudio::ResonanceAudioApi* GetResonanceApi()
+		{
+			return ResonanceSystem;
+		}
+	};
+
+	class FAmbisonicsFactory : public ISoundfieldFactory
+	{
+	public:
+		FAmbisonicsFactory();
+		virtual ~FAmbisonicsFactory();
 
 		// This function is called by the ResonanceAudioPluginListener as soon as the ResonanceAudioApi is initialized.
-		void SetResonanceAudioApi(vraudio::ResonanceAudioApi* InResonanceAudioApi) { ResonanceAudioApi = InResonanceAudioApi; };
+		void SetResonanceAudioApi(uint32 AudioDeviceID, vraudio::ResonanceAudioApi* InResonanceAudioApi) 
+		{
+			vraudio::ResonanceAudioApi*& ApiPtr = ResonanceAudioAPIMap.FindOrAdd(AudioDeviceID);
+			ApiPtr = InResonanceAudioApi;
+		};
+
+		// Begin IAmbisonicsMixer
+		virtual FName GetSoundfieldFormatName() override;
+		virtual TUniquePtr<ISoundfieldEncoderStream> CreateEncoderStream(const FAudioPluginInitializationParams& InitInfo, const ISoundfieldEncodingSettingsProxy& InitialSettings) override;
+		virtual TUniquePtr<ISoundfieldDecoderStream> CreateDecoderStream(const FAudioPluginInitializationParams& InitInfo, const ISoundfieldEncodingSettingsProxy& InitialSettings) override;
+		virtual TUniquePtr<ISoundfieldTranscodeStream> CreateTranscoderStream(const FName SourceFormat, const ISoundfieldEncodingSettingsProxy& InitialSourceSettings, const FName DestinationFormat, const ISoundfieldEncodingSettingsProxy& InitialDestinationSettings, const FAudioPluginInitializationParams& InitInfo) override;
+		virtual TUniquePtr<ISoundfieldMixerStream> CreateMixerStream(const ISoundfieldEncodingSettingsProxy& InitialSettings) override;
+		virtual TUniquePtr<ISoundfieldAudioPacket> CreateEmptyPacket() override;
+		virtual bool IsTranscodeRequiredBetweenSettings(const ISoundfieldEncodingSettingsProxy& SourceSettings, const ISoundfieldEncodingSettingsProxy& DestinationSettings) override;
+		virtual bool CanTranscodeFromSoundfieldFormat(FName SourceFormat, const ISoundfieldEncodingSettingsProxy& SourceEncodingSettings) override;
+		virtual bool CanTranscodeToSoundfieldFormat(FName DestinationFormat, const ISoundfieldEncodingSettingsProxy& DestinationEncodingSettings) override;
+		virtual UClass* GetCustomEncodingSettingsClass() const override;
+		virtual const USoundfieldEncodingSettingsBase* GetDefaultEncodingSettings() override;
+		//~ IAmbisonicsMixer
 
 	private:
-		// Map of which audio engine StreamIds map to which Resonance SourceIds.
-		TMap<uint32, RaSourceId> ResonanceSourceIdMap;
+		vraudio::ResonanceAudioApi* CreateNewResonanceApiInstance(FAudioDevice* AudioDevice, const FAudioPluginInitializationParams& InInitInfo);
 
-		//Handle to the API context used by this plugin.
-		vraudio::ResonanceAudioApi* ResonanceAudioApi;
+		// Map of which audio engine StreamIds map to which Resonance SourceIds.
+		TMap<uint32, vraudio::ResonanceAudioApi*> ResonanceAudioAPIMap;
 	};
 }
\ No newline at end of file
diff --git a/Engine/Plugins/Runtime/ResonanceAudio/Source/ResonanceAudio/Private/ResonanceAudioCommon.cpp b/Engine/Plugins/Runtime/ResonanceAudio/Source/ResonanceAudio/Private/ResonanceAudioCommon.cpp
--- a/Engine/Plugins/Runtime/ResonanceAudio/Source/ResonanceAudio/Private/ResonanceAudioCommon.cpp
+++ b/Engine/Plugins/Runtime/ResonanceAudio/Source/ResonanceAudio/Private/ResonanceAudioCommon.cpp
@@ -9,6 +9,7 @@
 #include "Interfaces/IPluginManager.h"
 #include "HAL/PlatformProcess.h"
 #include "Misc/Paths.h"
+#include "AudioMixerDevice.h"
 
 DEFINE_LOG_CATEGORY(LogResonanceAudio);
 
@@ -150,6 +151,16 @@ namespace ResonanceAudio
 	}
 
 
+	FVector ConvertToResonanceAudioCoordinates(const Audio::FChannelPositionInfo& ChannelPositionInfo)
+	{
+		FVector ResonanceAudioVector;
+		ResonanceAudioVector.X = ChannelPositionInfo.Radius * FMath::Sin(ChannelPositionInfo.Azimuth) * FMath::Cos(ChannelPositionInfo.Elevation);
+		ResonanceAudioVector.Y = ChannelPositionInfo.Radius * FMath::Sin(ChannelPositionInfo.Azimuth) * FMath::Sin(ChannelPositionInfo.Elevation);
+		ResonanceAudioVector.Z = ChannelPositionInfo.Radius * FMath::Cos(ChannelPositionInfo.Azimuth);
+
+		return ResonanceAudioVector * SCALE_FACTOR;
+	}
+
 	FQuat ConvertToResonanceAudioRotation(const FQuat& UnrealQuat)
 	{
 		FQuat ResonanceAudioQuat;
diff --git a/Engine/Plugins/Runtime/ResonanceAudio/Source/ResonanceAudio/Private/ResonanceAudioCommon.h b/Engine/Plugins/Runtime/ResonanceAudio/Source/ResonanceAudio/Private/ResonanceAudioCommon.h
--- a/Engine/Plugins/Runtime/ResonanceAudio/Source/ResonanceAudio/Private/ResonanceAudioCommon.h
+++ b/Engine/Plugins/Runtime/ResonanceAudio/Source/ResonanceAudio/Private/ResonanceAudioCommon.h
@@ -15,6 +15,7 @@ THIRD_PARTY_INCLUDES_END
 
 #include "UObject/ObjectMacros.h"
 #include "ResonanceAudioEnums.h"
+#include "ISoundfieldFormat.h"
 
 namespace ResonanceAudio
 {
@@ -38,6 +39,7 @@ namespace ResonanceAudio
 
 	// Converts between Unreal and Resonance Audio position coordinates.
 	FVector ConvertToResonanceAudioCoordinates(const FVector& UnrealVector);
+	FVector ConvertToResonanceAudioCoordinates(const Audio::FChannelPositionInfo& ChannelPositionInfo);
 
 	// Converts between Unreal and Resonance Audio rotation quaternions.
 	FQuat ConvertToResonanceAudioRotation(const FQuat& UnrealQuat);
diff --git a/Engine/Plugins/Runtime/ResonanceAudio/Source/ResonanceAudio/Private/ResonanceAudioModule.cpp b/Engine/Plugins/Runtime/ResonanceAudio/Source/ResonanceAudio/Private/ResonanceAudioModule.cpp
--- a/Engine/Plugins/Runtime/ResonanceAudio/Source/ResonanceAudio/Private/ResonanceAudioModule.cpp
+++ b/Engine/Plugins/Runtime/ResonanceAudio/Source/ResonanceAudio/Private/ResonanceAudioModule.cpp
@@ -111,12 +111,6 @@ namespace ResonanceAudio
 		return TAudioSpatializationPtr(new FResonanceAudioSpatialization());
 	}
 
-
-	TAmbisonicsMixerPtr FSpatializationPluginFactory::CreateNewAmbisonicsMixer(FAudioDevice* OwningDevice)
-	{
-		return TAmbisonicsMixerPtr(new FResonanceAudioAmbisonicsMixer());
-	}
-
 	TAudioReverbPtr FReverbPluginFactory::CreateNewReverbPlugin(FAudioDevice* OwningDevice)
 	{
 		// Register the audio device to the Resonance Audio module.
diff --git a/Engine/Plugins/Runtime/ResonanceAudio/Source/ResonanceAudio/Private/ResonanceAudioModule.h b/Engine/Plugins/Runtime/ResonanceAudio/Source/ResonanceAudio/Private/ResonanceAudioModule.h
--- a/Engine/Plugins/Runtime/ResonanceAudio/Source/ResonanceAudio/Private/ResonanceAudioModule.h
+++ b/Engine/Plugins/Runtime/ResonanceAudio/Source/ResonanceAudio/Private/ResonanceAudioModule.h
@@ -36,8 +36,6 @@ namespace ResonanceAudio
 		
 		virtual bool IsExternalSend() override { return true; }
 
-		virtual TAmbisonicsMixerPtr CreateNewAmbisonicsMixer(FAudioDevice* OwningDevice) override;
-
 		virtual int32 GetMaxSupportedChannels() override
 		{
 			return 2;
@@ -110,6 +108,7 @@ namespace ResonanceAudio
 		// Plugin factories.
 		FSpatializationPluginFactory SpatializationPluginFactory;
 		FReverbPluginFactory ReverbPluginFactory;
+		FAmbisonicsFactory AmbisonicsFactory;
 	};
 
 }  // namespace ResonanceAudio
diff --git a/Engine/Plugins/Runtime/ResonanceAudio/Source/ResonanceAudio/Private/ResonanceAudioPluginListener.cpp b/Engine/Plugins/Runtime/ResonanceAudio/Source/ResonanceAudio/Private/ResonanceAudioPluginListener.cpp
--- a/Engine/Plugins/Runtime/ResonanceAudio/Source/ResonanceAudio/Private/ResonanceAudioPluginListener.cpp
+++ b/Engine/Plugins/Runtime/ResonanceAudio/Source/ResonanceAudio/Private/ResonanceAudioPluginListener.cpp
@@ -13,10 +13,10 @@ namespace ResonanceAudio
 {
 	FResonanceAudioPluginListener::FResonanceAudioPluginListener()
 		: ResonanceAudioApi(nullptr)
+		, OwningAudioDevice(nullptr)
 		, ResonanceAudioModule(nullptr)
 		, ReverbPtr(nullptr)
 		, SpatializationPtr(nullptr)
-		, AmbisonicsPtr(nullptr)
 	{
 	}
 
@@ -26,6 +26,11 @@ namespace ResonanceAudio
 		{
 			delete ResonanceAudioApi;
 			ResonanceAudioApi = nullptr;
+
+			check(OwningAudioDevice);
+
+			FScopeLock ScopeLock(&ResonanceApiMapCriticalSection);
+			ResonanceApiMap.Remove(OwningAudioDevice);
 		}
 	}
 
@@ -46,10 +51,16 @@ namespace ResonanceAudio
 			UE_LOG(LogResonanceAudio, Error, TEXT("Failed to initialize Resonance Audio API"));
 			return;
 		}
+		else
+		{
+			FScopeLock ScopeLock(&ResonanceApiMapCriticalSection);
+			check(AudioDevice);
+			ResonanceApiMap.FindOrAdd(AudioDevice, ResonanceAudioApi);
+			OwningAudioDevice = AudioDevice;
+		}
 
 		ReverbPtr = (FResonanceAudioReverb*)AudioDevice->ReverbPluginInterface.Get();
 		SpatializationPtr = (FResonanceAudioSpatialization*)AudioDevice->SpatializationPluginInterface.Get();
-		AmbisonicsPtr = (FResonanceAudioAmbisonicsMixer*)AudioDevice->GetAmbisonicsMixer().Get();
 
 		// Make sure that Reverb *AND* spatialization plugins are enabled.
 		if (ReverbPtr == nullptr || SpatializationPtr == nullptr)
@@ -61,15 +72,6 @@ namespace ResonanceAudio
 		ReverbPtr->SetResonanceAudioApi(ResonanceAudioApi);
 		SpatializationPtr->SetResonanceAudioApi(ResonanceAudioApi);
 
-		if (AmbisonicsPtr == nullptr)
-		{
-			UE_LOG(LogResonanceAudio, Warning, TEXT("Resonance Audio Ambisonics plugin not found. Ambisonics files will not play correctly."));
-		}
-		else
-		{
-			AmbisonicsPtr->SetResonanceAudioApi(ResonanceAudioApi);
-		}
-
 		UE_LOG(LogResonanceAudio, Display, TEXT("Resonance Audio Listener is initialized"));
 	}
 
@@ -130,4 +132,53 @@ namespace ResonanceAudio
 		}
 	}
 
+	vraudio::ResonanceAudioApi* FResonanceAudioPluginListener::GetResonanceAPIForAudioDevice(const FAudioDevice* InAudioDevice)
+	{
+		FScopeLock ScopeLock(&ResonanceApiMapCriticalSection);
+
+		check(InAudioDevice);
+		if (vraudio::ResonanceAudioApi** ResonanceApiPtr = ResonanceApiMap.Find(InAudioDevice))
+		{
+			return *ResonanceApiPtr;
+		}
+		else
+		{
+			return nullptr;
+		}
+	}
+
+	void FResonanceAudioPluginListener::SetResonanceAPIForAudioDevice(const FAudioDevice* InAudioDevice, vraudio::ResonanceAudioApi* InResonanceSystem)
+	{
+		FScopeLock ScopeLock(&ResonanceApiMapCriticalSection);
+		ResonanceApiMap.Add(InAudioDevice, InResonanceSystem);
+	}
+
+	void FResonanceAudioPluginListener::RemoveResonanceAPIForAudioDevice(const FAudioDevice* InAudioDevice)
+	{
+		FScopeLock ScopeLock(&ResonanceApiMapCriticalSection);
+		ResonanceApiMap.Remove(InAudioDevice);
+	}
+
+	void FResonanceAudioPluginListener::RemoveResonanceAPIForAudioDevice(const vraudio::ResonanceAudioApi* InResonanceSystem)
+	{
+		FScopeLock ScopeLock(&ResonanceApiMapCriticalSection);
+
+		const FAudioDevice* AudioDeviceKey = nullptr;
+		for (auto& Pair : ResonanceApiMap)
+		{
+			if (Pair.Value == InResonanceSystem)
+			{
+				AudioDeviceKey = Pair.Key;
+				break;
+			}
+		}
+
+		checkf(AudioDeviceKey, TEXT("RemoveResonanceAPIForAudioDevice was called for a resonance system that wasn't registered using SetResonanceAPIForAudioDevice!"));
+		ResonanceApiMap.Remove(AudioDeviceKey);
+	}
+
+	TMap<const FAudioDevice*, vraudio::ResonanceAudioApi*> FResonanceAudioPluginListener::ResonanceApiMap;
+
+	FCriticalSection FResonanceAudioPluginListener::ResonanceApiMapCriticalSection;
+
 }  // namespace ResonanceAudio
diff --git a/Engine/Plugins/Runtime/ResonanceAudio/Source/ResonanceAudio/Private/ResonanceAudioPluginListener.h b/Engine/Plugins/Runtime/ResonanceAudio/Source/ResonanceAudio/Private/ResonanceAudioPluginListener.h
--- a/Engine/Plugins/Runtime/ResonanceAudio/Source/ResonanceAudio/Private/ResonanceAudioPluginListener.h
+++ b/Engine/Plugins/Runtime/ResonanceAudio/Source/ResonanceAudio/Private/ResonanceAudioPluginListener.h
@@ -13,7 +13,7 @@
 
 namespace ResonanceAudio
 {
-	// Dispatches listener information to the Resonance Audio plugins.
+	// Dispatches listener information to the Resonance Audio plugins and owns the vraudio::ResonanceAudioApi.
 	class FResonanceAudioPluginListener : public IAudioPluginListener
 	{
 	public:
@@ -25,14 +25,25 @@ namespace ResonanceAudio
 		virtual void OnListenerShutdown(FAudioDevice* AudioDevice) override;
 		virtual void OnTick(UWorld* InWorld, const int32 ViewportIndex, const FTransform& ListenerTransform, const float InDeltaSeconds) override;
 
+		static vraudio::ResonanceAudioApi* GetResonanceAPIForAudioDevice(const FAudioDevice* InAudioDevice);
+		static void SetResonanceAPIForAudioDevice(const FAudioDevice* InAudioDevice, vraudio::ResonanceAudioApi* InResonanceSystem);
+		static void RemoveResonanceAPIForAudioDevice(const FAudioDevice* InAudioDevice);
+		static void RemoveResonanceAPIForAudioDevice(const vraudio::ResonanceAudioApi* InResonanceSystem);
+
 	private:
 		// Resonance Audio API instance.
 		vraudio::ResonanceAudioApi* ResonanceAudioApi;
 
+		// Map of Resonance API systems to Audio Device IDs.
+		static TMap<const FAudioDevice*, vraudio::ResonanceAudioApi*> ResonanceApiMap;
+		static FCriticalSection ResonanceApiMapCriticalSection;
+
+		// This audio device ptr should only be used to remove the ResonanceAudioApi* from the ResonanceApiMap on destruction.
+		const FAudioDevice* OwningAudioDevice;
+
 		class FResonanceAudioModule* ResonanceAudioModule;
 		class FResonanceAudioReverb* ReverbPtr;
 		class FResonanceAudioSpatialization* SpatializationPtr;
-		class FResonanceAudioAmbisonicsMixer* AmbisonicsPtr;
 	};
 
 } // namespace ResonanceAudio
diff --git a/Engine/Plugins/Runtime/ResonanceAudio/Source/ResonanceAudio/Public/ResonanceAudioAmbisonicsSettings.h b/Engine/Plugins/Runtime/ResonanceAudio/Source/ResonanceAudio/Public/ResonanceAudioAmbisonicsSettings.h
--- a/Engine/Plugins/Runtime/ResonanceAudio/Source/ResonanceAudio/Public/ResonanceAudioAmbisonicsSettings.h
+++ b/Engine/Plugins/Runtime/ResonanceAudio/Source/ResonanceAudio/Public/ResonanceAudioAmbisonicsSettings.h
@@ -5,24 +5,100 @@
 #pragma once
 
 #include "ResonanceAudioCommon.h"
-#include "IAmbisonicsMixer.h"
+#include "ISoundfieldFormat.h"
 #include "ResonanceAudioAmbisonicsSettings.generated.h"
 
 UENUM()
-enum class EAmbisonicsOrder : uint8
+enum class EResonanceRenderMode : uint8
 {
-	FirstOrder,
-	SecondOrder,
-	ThirdOrder
+	// Stereo panning, i.e., this disables HRTF-based rendering.
+	StereoPanning,
+	// HRTF-based rendering using First Order Ambisonics, over a virtual array of
+	// 8 loudspeakers arranged in a cube configuration around the listener's head.
+	BinauralLowQuality,
+	// HRTF-based rendering using Second Order Ambisonics, over a virtual array of
+	// 12 loudspeakers arranged in a dodecahedral configuration (using faces of
+	// the dodecahedron).
+	BinauralMediumQuality,
+	// HRTF-based rendering using Third Order Ambisonics, over a virtual array of
+	// 26 loudspeakers arranged in a Lebedev grid: https://goo.gl/DX1wh3.
+	BinauralHighQuality,
+	// Room effects only rendering. This disables HRTF-based rendering and direct
+	// (dry) output of a sound object. Note that this rendering mode should *not*
+	// be used for general-purpose sound object spatialization, as it will only
+	// render the corresponding room effects of given sound objects without the
+	// direct spatialization.
+	RoomEffectsOnly,
+};
+
+class FResonanceAmbisonicsSettingsProxy : public ISoundfieldEncodingSettingsProxy
+{
+public:
+	vraudio::RenderingMode RenderingMode;
+
+
+	virtual uint32 GetUniqueId() const override
+	{
+		return static_cast<uint32>(RenderingMode);
+	}
+
+
+	virtual TUniquePtr<ISoundfieldEncodingSettingsProxy> Duplicate() const override
+	{
+		FResonanceAmbisonicsSettingsProxy* Proxy = new FResonanceAmbisonicsSettingsProxy();
+		Proxy->RenderingMode = RenderingMode;
+		return TUniquePtr<ISoundfieldEncodingSettingsProxy>(Proxy);
+	}
 };
 
 UCLASS()
-class RESONANCEAUDIO_API UResonanceAudioAmbisonicsSettings : public UAmbisonicsSubmixSettingsBase
+class RESONANCEAUDIO_API UResonanceAudioSoundfieldSettings : public USoundfieldEncodingSettingsBase
 {
 	GENERATED_BODY()
 
 public:
 	//Which order of ambisonics to use for this submix.
 	UPROPERTY(EditAnywhere, BlueprintReadOnly, Category = Ambisonics)
-	EAmbisonicsOrder AmbisonicsOrder;
+	EResonanceRenderMode RenderMode = EResonanceRenderMode::BinauralHighQuality;
+
+	virtual TUniquePtr<ISoundfieldEncodingSettingsProxy> GetProxy() const override
+	{
+		FResonanceAmbisonicsSettingsProxy* Proxy = new FResonanceAmbisonicsSettingsProxy();
+
+		switch (RenderMode)
+		{
+			case EResonanceRenderMode::StereoPanning:
+			{
+				Proxy->RenderingMode = vraudio::kStereoPanning;
+				break;
+			}
+			case EResonanceRenderMode::BinauralLowQuality:
+			{
+				Proxy->RenderingMode = vraudio::kBinauralLowQuality;
+				break;
+			}
+			case EResonanceRenderMode::BinauralMediumQuality:
+			{
+				Proxy->RenderingMode = vraudio::kBinauralMediumQuality;
+				break;
+			}
+			case EResonanceRenderMode::BinauralHighQuality:
+			{
+				Proxy->RenderingMode = vraudio::kBinauralHighQuality;
+				break;
+			}
+			case EResonanceRenderMode::RoomEffectsOnly:
+			{
+				Proxy->RenderingMode = vraudio::kRoomEffectsOnly;
+				break;
+			}
+			default:
+			{
+				checkNoEntry();
+			}
+		}
+
+		return TUniquePtr<ISoundfieldEncodingSettingsProxy>(Proxy);
+	}
+
 };
\ No newline at end of file
diff --git a/Engine/Plugins/Runtime/ResonanceAudio/Source/ResonanceAudio/ResonanceAudio.Build.cs b/Engine/Plugins/Runtime/ResonanceAudio/Source/ResonanceAudio/ResonanceAudio.Build.cs
--- a/Engine/Plugins/Runtime/ResonanceAudio/Source/ResonanceAudio/ResonanceAudio.Build.cs
+++ b/Engine/Plugins/Runtime/ResonanceAudio/Source/ResonanceAudio/ResonanceAudio.Build.cs
@@ -34,6 +34,7 @@ public class ResonanceAudio : ModuleRules
                 "CoreUObject",
                 "Engine",
                 "AudioMixer",
+				"SoundFieldRendering",
                 "ProceduralMeshComponent",
             }
         );
@@ -52,7 +53,8 @@ public class ResonanceAudio : ModuleRules
                 "InputCore",
                 "Projects",
                 "AudioMixer",
-                "ProceduralMeshComponent"
+                "ProceduralMeshComponent",
+                "AudioExtensions"
             }
         );
 
diff --git a/Engine/Plugins/Runtime/SoundFields/Source/SoundFields/Private/SoundFieldsModule.cpp b/Engine/Plugins/Runtime/SoundFields/Source/SoundFields/Private/SoundFieldsModule.cpp
new file mode 100644
--- /dev/null
+++ b/Engine/Plugins/Runtime/SoundFields/Source/SoundFields/Private/SoundFieldsModule.cpp
@@ -0,0 +1,16 @@
+// Copyright 1998-2019 Epic Games, Inc. All Rights Reserved.
+
+#include "SoundFieldsModule.h"
+#include "Features/IModularFeatures.h"
+
+
+void FSoundFieldsModule::StartupModule()
+{
+}
+
+void FSoundFieldsModule::ShutdownModule()
+{
+	// Nothing done here for now.
+}
+
+IMPLEMENT_MODULE(FSoundFieldsModule, SoundFields)
\ No newline at end of file
diff --git a/Engine/Plugins/Runtime/SoundFields/Source/SoundFields/Public/SoundFieldsModule.h b/Engine/Plugins/Runtime/SoundFields/Source/SoundFields/Public/SoundFieldsModule.h
new file mode 100644
--- /dev/null
+++ b/Engine/Plugins/Runtime/SoundFields/Source/SoundFields/Public/SoundFieldsModule.h
@@ -0,0 +1,19 @@
+// Copyright 1998-2019 Epic Games, Inc. All Rights Reserved.
+
+#pragma once
+
+
+#include "CoreMinimal.h"
+#include "Modules/ModuleManager.h"
+#include "SoundFields.h"
+
+class FSoundFieldsModule : public IModuleInterface
+{
+public:
+	virtual void StartupModule() override;
+
+	virtual void ShutdownModule() override;
+
+private:
+	FAmbisonicsSoundfieldFormat SFFactory;
+};
diff --git a/Engine/Plugins/Runtime/Spatialization/Source/Spatialization/Spatialization.Build.cs b/Engine/Plugins/Runtime/Spatialization/Source/Spatialization/Spatialization.Build.cs
--- a/Engine/Plugins/Runtime/Spatialization/Source/Spatialization/Spatialization.Build.cs
+++ b/Engine/Plugins/Runtime/Spatialization/Source/Spatialization/Spatialization.Build.cs
@@ -23,7 +23,8 @@ public class Spatialization : ModuleRules
 			{
 				"CoreUObject",
 				"Engine",
-				"SignalProcessing"
+				"SignalProcessing",
+				"AudioExtensions"
 			}
 			);
 
diff --git a/Engine/Plugins/Runtime/Steam/SteamAudio/Source/SteamAudio/SteamAudio.Build.cs b/Engine/Plugins/Runtime/Steam/SteamAudio/Source/SteamAudio/SteamAudio.Build.cs
--- a/Engine/Plugins/Runtime/Steam/SteamAudio/Source/SteamAudio/SteamAudio.Build.cs
+++ b/Engine/Plugins/Runtime/Steam/SteamAudio/Source/SteamAudio/SteamAudio.Build.cs
@@ -41,7 +41,8 @@ namespace UnrealBuildTool.Rules
 					"Engine",
 					"InputCore",
 					"Projects",
-					"AudioMixer"
+					"AudioMixer",
+					"AudioExtensions"
 				}
 			);
 
