For clarity, prefix oculus audio logs with Oculus Audio.

diff --git a/Engine/Plugins/Runtime/Oculus/OculusAudio/Source/OculusAudio/Private/OculusAudioDllManager.cpp b/Engine/Plugins/Runtime/Oculus/OculusAudio/Source/OculusAudio/Private/OculusAudioDllManager.cpp
--- a/Engine/Plugins/Runtime/Oculus/OculusAudio/Source/OculusAudio/Private/OculusAudioDllManager.cpp
+++ b/Engine/Plugins/Runtime/Oculus/OculusAudio/Source/OculusAudio/Private/OculusAudioDllManager.cpp
@@ -40,7 +40,7 @@ void FOculusAudioLibraryManager::Initialize()
 	{
 		if (!LoadDll())
 		{
-			UE_LOG(LogAudio, Error, TEXT("Failed to load OVR Audio dll"));
+			UE_LOG(LogAudio, Error, TEXT("Oculus Audio: Failed to load OVR Audio dll"));
 			check(false);
 			return;
 		}
@@ -57,7 +57,7 @@ void FOculusAudioLibraryManager::Initialize()
 		const char* OvrVersionString = OVRA_CALL(ovrAudio_GetVersion)(&MajorVersionNumber, &MinorVersionNumber, &PatchNumber);
 		if (MajorVersionNumber != OVR_AUDIO_MAJOR_VERSION || MinorVersionNumber != OVR_AUDIO_MINOR_VERSION)
 		{
-			UE_LOG(LogAudio, Warning, TEXT("Using mismatched OVR Audio SDK Version! %d.%d vs. %d.%d"), OVR_AUDIO_MAJOR_VERSION, OVR_AUDIO_MINOR_VERSION, MajorVersionNumber, MinorVersionNumber);
+			UE_LOG(LogAudio, Warning, TEXT("Oculus Audio: Using mismatched OVR Audio SDK Version! %d.%d vs. %d.%d"), OVR_AUDIO_MAJOR_VERSION, OVR_AUDIO_MINOR_VERSION, MajorVersionNumber, MinorVersionNumber);
 			return;
 		}
 		bInitialized = true;
@@ -110,25 +110,25 @@ bool FOculusAudioLibraryManager::LoadDll()
 			ClientType = OVRA_CLIENT_TYPE_WWISE_UNKNOWN;
 			Path = WwisePath;
 			DLL_NAME = WWISE_DLL_NAME;
-			UE_LOG(LogAudio, Display, TEXT("OculusSpatializerWwise.dll found, using the Wwise version of the Oculus Audio UE4 integration"));
+			UE_LOG(LogAudio, Display, TEXT("Oculus Audio: OculusSpatializerWwise.dll found, using the Wwise version of the Oculus Audio UE4 integration"));
 		}
 		else if (FPaths::FileExists(FMODPath + FMOD_DLL_NAME + ".dll"))
 		{
 			ClientType = OVRA_CLIENT_TYPE_FMOD;
 			Path = FMODPath;
 			DLL_NAME = FMOD_DLL_NAME;
-			UE_LOG(LogAudio, Display, TEXT("OculusSpatializerFMOD.dll found, using the FMOD version of the Oculus Audio UE4 integration"));
+			UE_LOG(LogAudio, Display, TEXT("Oculus Audio: OculusSpatializerFMOD.dll found, using the FMOD version of the Oculus Audio UE4 integration"));
 		}
 		else
 		{
 			ClientType = -1;
 			Path = UE4Path;
 			DLL_NAME = UE4_DLL_NAME;
-			UE_LOG(LogAudio, Display, TEXT("Middleware plugins not found, assuming native UE4 AudioMixer"));
+			UE_LOG(LogAudio, Display, TEXT("Oculus Audio: Middleware plugins not found, assuming native UE4 AudioMixer"));
 		}
 
 
-		UE_LOG(LogAudio, Display, TEXT("Attempting to load Oculus Spatializer DLL: %s (from %s)"), *Path, DLL_NAME);
+		UE_LOG(LogAudio, Display, TEXT("Oculus Audio: Attempting to load Oculus Spatializer DLL: %s (from %s)"), *Path, DLL_NAME);
 
 		FPlatformProcess::PushDllDirectory(*Path);
 		OculusAudioDllHandle = FPlatformProcess::GetDllHandle(*(Path + DLL_NAME));
@@ -139,24 +139,24 @@ bool FOculusAudioLibraryManager::LoadDll()
 		OculusAudioDllHandle = FPlatformProcess::GetDllHandle(*(Path + WWISE_DLL_NAME + ".so"));
 		if (OculusAudioDllHandle != nullptr)
 		{
-			UE_LOG(LogAudio, Display, TEXT("%s found, using the Wwise version of the Oculus Audio UE4 integration"), WWISE_DLL_NAME);
+			UE_LOG(LogAudio, Display, TEXT("Oculus Audio: %s found, using the Wwise version of the Oculus Audio UE4 integration"), WWISE_DLL_NAME);
 			return true;
 		}
 		OculusAudioDllHandle = FPlatformProcess::GetDllHandle(*(Path + FMOD_DLL_NAME + ".so"));
 		if (OculusAudioDllHandle != nullptr)
 		{
-			UE_LOG(LogAudio, Display, TEXT("%s found, using the FMOD version of the Oculus Audio UE4 integration"), FMOD_DLL_NAME);
+			UE_LOG(LogAudio, Display, TEXT("Oculus Audio: %s found, using the FMOD version of the Oculus Audio UE4 integration"), FMOD_DLL_NAME);
 			return true;
 		}
 		OculusAudioDllHandle = FPlatformProcess::GetDllHandle(*(Path + UE4_DLL_NAME + ".so"));
 		if (OculusAudioDllHandle != nullptr)
 		{
-			UE_LOG(LogAudio, Display, TEXT("Middleware plugins not found, %s found, assuming native UE4 AudioMixer"), UE4_DLL_NAME);
+			UE_LOG(LogAudio, Display, TEXT("Oculus Audio: Middleware plugins not found, %s found, assuming native UE4 AudioMixer"), UE4_DLL_NAME);
 			return true;
 		}
 		else
 		{
-			UE_LOG(LogAudio, Error, TEXT("Unable to load Oculus Audio UE4 integratiton"), UE4_DLL_NAME);
+			UE_LOG(LogAudio, Error, TEXT("Oculus Audio: Unable to load Oculus Audio UE4 integratiton"), UE4_DLL_NAME);
 		}
 #endif
 
