Fix for distorted Oculus audio, was using InitInfo.NumChannels, which is actually the number of channels of the incoming source audio, instead of FAudioDevice::MaxChannels, which is the max number of sources.

diff --git a/Engine/Plugins/Runtime/Oculus/OculusAudio/Source/OculusAudio/Private/OculusAmbisonicSpatializer.cpp b/Engine/Plugins/Runtime/Oculus/OculusAudio/Source/OculusAudio/Private/OculusAmbisonicSpatializer.cpp
--- a/Engine/Plugins/Runtime/Oculus/OculusAudio/Source/OculusAudio/Private/OculusAmbisonicSpatializer.cpp
+++ b/Engine/Plugins/Runtime/Oculus/OculusAudio/Source/OculusAudio/Private/OculusAmbisonicSpatializer.cpp
@@ -349,7 +349,7 @@ TUniquePtr<ISoundfieldEncoderStream> FOculusAmbisonicsFactory::CreateEncoderStre
 		check(Context);
 	}
 
-	return TUniquePtr<ISoundfieldEncoderStream>(new FOculusEncoder(Context, InitInfo.NumOutputChannels, InitInfo.NumSources));
+	return TUniquePtr<ISoundfieldEncoderStream>(new FOculusEncoder(Context, InitInfo.NumOutputChannels, InitInfo.AudioDevicePtr->GetMaxChannels()));
 }
 
 TUniquePtr<ISoundfieldDecoderStream> FOculusAmbisonicsFactory::CreateDecoderStream(const FAudioPluginInitializationParams& InitInfo, const ISoundfieldEncodingSettingsProxy& InitialSettings)
