Fix for oculus or resonance submixes not outputting audible audio.

diff --git a/Engine/Plugins/Runtime/Oculus/OculusAudio/Source/OculusAudio/Private/OculusAmbisonicSpatializer.cpp b/Engine/Plugins/Runtime/Oculus/OculusAudio/Source/OculusAudio/Private/OculusAmbisonicSpatializer.cpp
--- a/Engine/Plugins/Runtime/Oculus/OculusAudio/Source/OculusAudio/Private/OculusAmbisonicSpatializer.cpp
+++ b/Engine/Plugins/Runtime/Oculus/OculusAudio/Source/OculusAudio/Private/OculusAmbisonicSpatializer.cpp
@@ -8,14 +8,31 @@
 #include "OculusAudioDllManager.h"
 #include "OculusAudioContextManager.h"
 #include "SoundFieldRendering.h"
+#include "AudioPluginUtilities.h"
 
-static int32 UseSoundfieldSubmixForOculusAudioOutputCVar = 0;
-FAutoConsoleVariableRef CVarUseSoundfieldSubmixForOculusAudioOutput(
-	TEXT("au.oculus.UseSoundfieldSubmixForOculusAudioOutput"),
-	UseSoundfieldSubmixForOculusAudioOutputCVar,
-	TEXT("If set to 1, this will allow for audio sources to be sent to an Oculus Audio soundfield submix to be HRTF spatialized, rather than using the Oculus Audio Spatialization Plugin for the entire project.\n")
-	TEXT("0: Disable, >0: Enable"),
-	ECVF_Default);
+namespace OculusAudioUtils
+{
+	bool IsOculusAudioTheCurrentSpatializationPlugin()
+	{
+#if WITH_EDITOR
+		static FString OculusDisplayName = FString(TEXT("Oculus Audio"));
+		return AudioPluginUtilities::GetDesiredPluginName(EAudioPlugin::SPATIALIZATION).Equals(OculusDisplayName);
+#else
+		// For non editor situations, we can cache whether this is the current plugin or not the first time we check.
+		static FString OculusDisplayName = FString(TEXT("Oculus Audio"));
+		static bool bCheckedSpatializationPlugin = false;
+		static bool bIsOculusCurrentSpatiatizationPlugin = false;
+
+		if (!bCheckedSpatializationPlugin)
+		{
+			bIsOculusCurrentSpatiatizationPlugin = AudioPluginUtilities::GetDesiredPluginName(EAudioPlugin::SPATIALIZATION).Equals(OculusDisplayName);
+			bCheckedSpatializationPlugin = true;
+		}
+
+		return bIsOculusCurrentSpatiatizationPlugin;
+#endif
+	}
+}
 
 // the Oculus Audio renderer renders all inputs immediately to an interleaved stereo output.
 class FOculusSoundfieldBuffer : public ISoundfieldAudioPacket
@@ -105,8 +122,8 @@ public:
 	}
 
 
-	// Binaurally spatialize each independent speaker position. If UseSoundfieldSubmixForOculusAudioOutputCVar
-	// was false when this encoder was created this will no-op.
+	// Binaurally spatialize each independent speaker position. If the spatialization plugin for the current platform was set to Oculus Audio
+	// when this encoder was created this will no-op.
 	virtual void EncodeAndMixIn(const FSoundfieldEncoderInputData& InputData, ISoundfieldAudioPacket& OutputData) override
 	{
 		FOculusSoundfieldBuffer& OutputBuffer = DowncastSoundfieldRef<FOculusSoundfieldBuffer>(OutputData);
@@ -175,7 +192,15 @@ FCriticalSection FOculusEncoder::SourceIdCounterCritSection;
 // we simply pass it forward here.
 class FOculusAudioDecoder : public ISoundfieldDecoderStream
 {
+private:
+	const bool bShouldUseSubmixForOutput;
+
 public:
+	FOculusAudioDecoder(const bool bInShouldUseSubmixForOutput)
+		: bShouldUseSubmixForOutput(bInShouldUseSubmixForOutput)
+	{}
+
+	FOculusAudioDecoder() = delete;
 
 	virtual void Decode(const FSoundfieldDecoderInputData& InputData, FSoundfieldDecoderOutputData& OutputData) override
 	{
@@ -186,7 +211,7 @@ public:
 
 	virtual void DecodeAndMixIn(const FSoundfieldDecoderInputData& InputData, FSoundfieldDecoderOutputData& OutputData) override
 	{
-		if (UseSoundfieldSubmixForOculusAudioOutputCVar)
+		if (bShouldUseSubmixForOutput)
 		{
 			const FOculusSoundfieldBuffer& InputBuffer = DowncastSoundfieldRef<const FOculusSoundfieldBuffer>(InputData.SoundfieldBuffer);
 
@@ -354,7 +379,11 @@ TUniquePtr<ISoundfieldEncoderStream> FOculusAmbisonicsFactory::CreateEncoderStre
 
 TUniquePtr<ISoundfieldDecoderStream> FOculusAmbisonicsFactory::CreateDecoderStream(const FAudioPluginInitializationParams& InitInfo, const ISoundfieldEncodingSettingsProxy& InitialSettings)
 {
-	return TUniquePtr<ISoundfieldDecoderStream>(new FOculusAudioDecoder());
+	const bool bShouldOutputOculusAudio = !OculusAudioUtils::IsOculusAudioTheCurrentSpatializationPlugin();
+
+	UE_CLOG(!bShouldOutputOculusAudio, LogAudio, Warning, TEXT("Oculus Audio is the currently selected spatialization plugin. The Oculus soundfield submix's output will be ignored."));
+
+	return TUniquePtr<ISoundfieldDecoderStream>(new FOculusAudioDecoder(bShouldOutputOculusAudio));
 }
 
 TUniquePtr<ISoundfieldTranscodeStream> FOculusAmbisonicsFactory::CreateTranscoderStream(const FName SourceFormat, const ISoundfieldEncodingSettingsProxy& InitialSourceSettings, const FName DestinationFormat, const ISoundfieldEncodingSettingsProxy& InitialDestinationSettings, const FAudioPluginInitializationParams& InitInfo)
diff --git a/Engine/Plugins/Runtime/ResonanceAudio/Source/ResonanceAudio/Private/ResonanceAudioAmbisonics.cpp b/Engine/Plugins/Runtime/ResonanceAudio/Source/ResonanceAudio/Private/ResonanceAudioAmbisonics.cpp
--- a/Engine/Plugins/Runtime/ResonanceAudio/Source/ResonanceAudio/Private/ResonanceAudioAmbisonics.cpp
+++ b/Engine/Plugins/Runtime/ResonanceAudio/Source/ResonanceAudio/Private/ResonanceAudioAmbisonics.cpp
@@ -5,14 +5,31 @@
 #include "ResonanceAudioAmbisonicsSettings.h"
 #include "ResonanceAudioPluginListener.h"
 #include "SoundFieldRendering.h"
+#include "AudioPluginUtilities.h"
 
-static int32 UseReverbPluginForResonanceOutputCVar = 1;
-FAutoConsoleVariableRef CVarUseReverbPluginForResonanceOutput(
-	TEXT("au.resonance.UseReverbPluginForOutput"),
-	UseReverbPluginForResonanceOutputCVar,
-	TEXT("If set to 0, will allow for Resonance to render it's output directly from a submix even if it is not set as the spatialization or reverb plugin.\n")
-	TEXT("0: Disable, >0: Enable"),
-	ECVF_Default);
+namespace ResonanceAudioUtils
+{
+	bool IsResonanceAudioTheCurrentSpatializationPlugin()
+	{
+#if WITH_EDITOR
+		static FString ResonanceDisplayName = FString(TEXT("Resonance Audio"));
+		return AudioPluginUtilities::GetDesiredPluginName(EAudioPlugin::REVERB).Equals(ResonanceDisplayName);
+#else
+		// For non editor situations, we can cache whether this is the current plugin or not the first time we check.
+		static FString ResonanceDisplayName = FString(TEXT("Resonance Audio"));
+		static bool bCheckedSpatializationPlugin = false;
+		static bool bIsResonanceCurrentSpatiatizationPlugin = false;
+		
+		if (!bCheckedSpatializationPlugin)
+		{
+			bIsResonanceCurrentSpatiatizationPlugin = AudioPluginUtilities::GetDesiredPluginName(EAudioPlugin::REVERB).Equals(ResonanceDisplayName);
+			bCheckedSpatializationPlugin = true;
+		}
+
+		return bIsResonanceCurrentSpatiatizationPlugin;
+#endif
+	}
+}
 
 namespace ResonanceAudio
 {
@@ -141,7 +158,17 @@ namespace ResonanceAudio
 
 	class FDecoder : public ISoundfieldDecoderStream
 	{
+	private:
+		const bool bShouldOutputAudio;
+
 	public:
+
+		FDecoder(bool bInShouldOutputAudio)
+			: bShouldOutputAudio(bInShouldOutputAudio)
+		{}
+
+		FDecoder() = delete;
+
 		virtual void Decode(const FSoundfieldDecoderInputData& InputData, FSoundfieldDecoderOutputData& OutputData) override
 		{
 			DecodeAndMixIn(InputData, OutputData);
@@ -151,8 +178,7 @@ namespace ResonanceAudio
 		virtual void DecodeAndMixIn(const FSoundfieldDecoderInputData& InputData, FSoundfieldDecoderOutputData& OutputData) override
 		{
 			// If the reverb plugin is enabled, FillInterleavedOutputBuffer will be called there.
-			// TODO: detect whether the reverb plugin is enabled on this audio device to toggle this behavior automatically.
-			if (!UseReverbPluginForResonanceOutputCVar)
+			if (bShouldOutputAudio)
 			{
 				const int32 NumChannels = InputData.PositionalData.NumChannels;
 				const FResonancePacket& Packet = DowncastSoundfieldRef<const FResonancePacket>(InputData.SoundfieldBuffer);
@@ -292,7 +318,9 @@ namespace ResonanceAudio
 
 	TUniquePtr<ISoundfieldDecoderStream> FAmbisonicsFactory::CreateDecoderStream(const FAudioPluginInitializationParams& InitInfo, const ISoundfieldEncodingSettingsProxy& InitialSettings)
 	{
-		return TUniquePtr<ISoundfieldDecoderStream>(new FDecoder());
+		const bool bShouldOutputResonanceAudioThroughSubmix = !ResonanceAudioUtils::IsResonanceAudioTheCurrentSpatializationPlugin();
+		UE_CLOG(!bShouldOutputResonanceAudioThroughSubmix, LogAudio, Display, TEXT("Since Resonance Audio is the currently selected spatialization plugin. The output of this plugin will be sent through the Reverb Plugin submix."));
+		return TUniquePtr<ISoundfieldDecoderStream>(new FDecoder(bShouldOutputResonanceAudioThroughSubmix));
 	}
 
 	TUniquePtr<ISoundfieldTranscodeStream> FAmbisonicsFactory::CreateTranscoderStream(const FName SourceFormat, const ISoundfieldEncodingSettingsProxy& InitialSourceSettings, const FName DestinationFormat, const ISoundfieldEncodingSettingsProxy& InitialDestinationSettings, const FAudioPluginInitializationParams& InitInfo)
