Missing generated.h
#rnx

diff --git a/Engine/Plugins/Runtime/Oculus/OculusAudio/Source/OculusAudio/Private/OculusAudioReverb.h b/Engine/Plugins/Runtime/Oculus/OculusAudio/Source/OculusAudio/Private/OculusAudioReverb.h
--- a/Engine/Plugins/Runtime/Oculus/OculusAudio/Source/OculusAudio/Private/OculusAudioReverb.h
+++ b/Engine/Plugins/Runtime/Oculus/OculusAudio/Source/OculusAudio/Private/OculusAudioReverb.h
@@ -9,6 +9,8 @@
 #include "Templates/UniquePtr.h"
 #include "Sound/SoundEffectBase.h"
 
+#include "OculusAudioReverb.generated.h"
+
 
 UCLASS()
 class USubmixEffectOculusReverbPreset : public USoundEffectSubmixPreset
