Missing gen.h round 2
#rnx

diff --git a/Engine/Plugins/Runtime/Oculus/OculusAudio/Source/OculusAudio/Private/OculusAudioReverb.h b/Engine/Plugins/Runtime/Oculus/OculusAudio/Source/OculusAudio/Private/OculusAudioReverb.h
--- a/Engine/Plugins/Runtime/Oculus/OculusAudio/Source/OculusAudio/Private/OculusAudioReverb.h
+++ b/Engine/Plugins/Runtime/Oculus/OculusAudio/Source/OculusAudio/Private/OculusAudioReverb.h
@@ -3,11 +3,12 @@
 #pragma once
 
 #include "IAudioExtensionPlugin.h"
-#include "Sound/SoundEffectSubmix.h"
-#include "ObjectMacros.h"
 #include "OVR_Audio.h"
-#include "Templates/UniquePtr.h"
 #include "Sound/SoundEffectBase.h"
+#include "Sound/SoundEffectSubmix.h"
+#include "Templates/UniquePtr.h"
+#include "UObject/Object.h"
+#include "UObject/ObjectMacros.h"
 
 #include "OculusAudioReverb.generated.h"
 
