Fix stubbed submix effect API in Oculus to compile
#rnx

diff --git a/Engine/Plugins/Runtime/Oculus/OculusAudio/Source/OculusAudio/Private/OculusAudioReverb.cpp b/Engine/Plugins/Runtime/Oculus/OculusAudio/Source/OculusAudio/Private/OculusAudioReverb.cpp
--- a/Engine/Plugins/Runtime/Oculus/OculusAudio/Source/OculusAudio/Private/OculusAudioReverb.cpp
+++ b/Engine/Plugins/Runtime/Oculus/OculusAudio/Source/OculusAudio/Private/OculusAudioReverb.cpp
@@ -70,12 +70,11 @@ FSoundEffectSubmixPtr OculusAudioReverb::GetEffectSubmix()
 	{
 		if (!ReverbPreset)
 		{
-			ReverbPreset = NewObject<USubmixEffectOculusReverbPreset>();
+			ReverbPreset = NewObject<USubmixEffectOculusReverbPluginPreset>();
 			ReverbPreset->AddToRoot();
 		}
 
-		FSoundEffectSubmixInitData InitData;
-		SubmixEffect = FSoundEffectSubmix::Create(InitData, *ReverbPreset);
+		SubmixEffect = USoundEffectPreset::CreateInstance<FSoundEffectSubmixInitData, FSoundEffectSubmix>(FSoundEffectSubmixInitData(), *ReverbPreset);
 		SubmixEffect->SetEnabled(true);
 	}
 
diff --git a/Engine/Plugins/Runtime/Oculus/OculusAudio/Source/OculusAudio/Private/OculusAudioReverb.h b/Engine/Plugins/Runtime/Oculus/OculusAudio/Source/OculusAudio/Private/OculusAudioReverb.h
--- a/Engine/Plugins/Runtime/Oculus/OculusAudio/Source/OculusAudio/Private/OculusAudioReverb.h
+++ b/Engine/Plugins/Runtime/Oculus/OculusAudio/Source/OculusAudio/Private/OculusAudioReverb.h
@@ -13,19 +13,11 @@
 #include "OculusAudioReverb.generated.h"
 
 
-UCLASS()
-class USubmixEffectOculusReverbPreset : public USoundEffectSubmixPreset
-{
-	GENERATED_BODY()
-
-public:
-	EFFECT_PRESET_METHODS(SubmixEffectOculusReverbPreset)
-};
+// Forward Declarations
+class USubmixEffectOculusReverbPluginPreset;
 
 class FSubmixEffectOculusReverbPlugin : public FSoundEffectSubmix
 {
-	void ClearContext();
-
 	virtual void Init(const FSoundEffectSubmixInitData& InInitData) override;
 	virtual uint32 GetDesiredInputChannelCountOverride() const override
 	{
@@ -38,9 +30,11 @@ class FSubmixEffectOculusReverbPlugin : public FSoundEffectSubmix
 		return; // PAS
 	}
 
-protected:
+public:
 	FSubmixEffectOculusReverbPlugin();
 
+	void ClearContext();
+
 private:
 	ovrAudioContext Context;
 	FCriticalSection ContextLock;
@@ -84,5 +78,30 @@ public:
 private:
 	ovrAudioContext* Context;
 	TSoundEffectSubmixPtr SubmixEffect;
-	USubmixEffectOculusReverbPreset* ReverbPreset;
+	USubmixEffectOculusReverbPluginPreset* ReverbPreset;
+};
+
+USTRUCT()
+struct OCULUSAUDIO_API FSubmixEffectOculusReverbPluginSettings
+{
+	GENERATED_USTRUCT_BODY()
+
+	FSubmixEffectOculusReverbPluginSettings() = default;
+};
+
+UCLASS()
+class OCULUSAUDIO_API USubmixEffectOculusReverbPluginPreset : public USoundEffectSubmixPreset
+{
+	GENERATED_BODY()
+
+public:
+	EFFECT_PRESET_METHODS(SubmixEffectOculusReverbPlugin)
+
+	UFUNCTION()
+	void SetSettings(const FSubmixEffectOculusReverbPluginSettings& InSettings)
+	{
+	}
+
+	UPROPERTY()
+	FSubmixEffectOculusReverbPluginSettings Settings;
 };
\ No newline at end of file
