Fix factory name copy pasta #rnx #jira none #rb none

diff --git a/Engine/Plugins/Runtime/Oculus/OculusAudio/Source/OculusAudioEditor/Private/OculusAmbisonicSettingsFactory.cpp b/Engine/Plugins/Runtime/Oculus/OculusAudio/Source/OculusAudioEditor/Private/OculusAmbisonicSettingsFactory.cpp
--- a/Engine/Plugins/Runtime/Oculus/OculusAudio/Source/OculusAudioEditor/Private/OculusAmbisonicSettingsFactory.cpp
+++ b/Engine/Plugins/Runtime/Oculus/OculusAudio/Source/OculusAudioEditor/Private/OculusAmbisonicSettingsFactory.cpp
@@ -10,12 +10,12 @@
 
 FText FAssetTypeActions_OculusAmbisonicsSettings::GetName() const
 {
-    return LOCTEXT("AssetTypeActions_OculusAmbisonicsSettings", "Oculus Audio Source Settings");
+	return LOCTEXT("AssetTypeActions_OculusAmbisonicsSettings", "Oculus Ambisonics Settings");
 }
 
 FColor FAssetTypeActions_OculusAmbisonicsSettings::GetTypeColor() const
 {
-    return FColor(100, 100, 100);
+	return FColor(100, 100, 100);
 }
 
 const TArray<FText>& FAssetTypeActions_OculusAmbisonicsSettings::GetSubMenus() const
@@ -30,12 +30,12 @@ const TArray<FText>& FAssetTypeActions_OculusAmbisonicsSettings::GetSubMenus() c
 
 UClass* FAssetTypeActions_OculusAmbisonicsSettings::GetSupportedClass() const
 {
-    return UOculusAudioSoundfieldSettings::StaticClass();
+	return UOculusAudioSoundfieldSettings::StaticClass();
 }
 
 uint32 FAssetTypeActions_OculusAmbisonicsSettings::GetCategories()
 {
-    return EAssetTypeCategories::Sounds;
+	return EAssetTypeCategories::Sounds;
 }
 
 UOculusAmbisonicsSettingsFactory::UOculusAmbisonicsSettingsFactory(const FObjectInitializer& ObjectInitializer)
@@ -49,13 +49,13 @@ UOculusAmbisonicsSettingsFactory::UOculusAmbisonicsSettingsFactory(const FObject
 }
 
 UObject* UOculusAmbisonicsSettingsFactory::FactoryCreateNew(UClass* Class, UObject* InParent, FName Name, EObjectFlags Flags, UObject* Context,
-    FFeedbackContext* Warn)
+	FFeedbackContext* Warn)
 {
-    return NewObject<UOculusAudioSoundfieldSettings>(InParent, Name, Flags);
+	return NewObject<UOculusAudioSoundfieldSettings>(InParent, Name, Flags);
 }
 
 uint32 UOculusAmbisonicsSettingsFactory::GetMenuCategories() const
 {
-    return EAssetTypeCategories::Sounds;
+	return EAssetTypeCategories::Sounds;
 }
 #undef LOCTEXT_NAMESPACE
