Div by Zero fix for Oculus sound field transcoder (empty input buffer)

diff --git a/Engine/Plugins/Runtime/Oculus/OculusAudio/Source/OculusAudio/Private/OculusAmbisonicSpatializer.cpp b/Engine/Plugins/Runtime/Oculus/OculusAudio/Source/OculusAudio/Private/OculusAmbisonicSpatializer.cpp
--- a/Engine/Plugins/Runtime/Oculus/OculusAudio/Source/OculusAudio/Private/OculusAmbisonicSpatializer.cpp
+++ b/Engine/Plugins/Runtime/Oculus/OculusAudio/Source/OculusAudio/Private/OculusAmbisonicSpatializer.cpp
@@ -263,6 +263,11 @@ public:
 		const FAmbisonicsSoundfieldBuffer& InputBuffer = DowncastSoundfieldRef<const FAmbisonicsSoundfieldBuffer>(InputData);
 		FOculusSoundfieldBuffer& OutputBuffer = DowncastSoundfieldRef<FOculusSoundfieldBuffer>(PacketToSumTo);
 
+		if (!InputBuffer.NumChannels || InputBuffer.AudioBuffer.IsEmpty())
+		{
+			return;
+		}
+
 		const int32 NumFrames = InputBuffer.AudioBuffer.Num() / InputBuffer.NumChannels;
 		OutputBuffer.AudioBuffer.SetNumZeroed(NumFrames * 2);
 		OutputBuffer.NumChannels = 2;
