Fixed compilation for Android with LWC on

diff --git a/Engine/Plugins/Runtime/Oculus/OculusAudio/Source/OculusAudio/Private/OculusAudioGeometryComponent.cpp b/Engine/Plugins/Runtime/Oculus/OculusAudio/Source/OculusAudio/Private/OculusAudioGeometryComponent.cpp
--- a/Engine/Plugins/Runtime/Oculus/OculusAudio/Source/OculusAudio/Private/OculusAudioGeometryComponent.cpp
+++ b/Engine/Plugins/Runtime/Oculus/OculusAudio/Source/OculusAudio/Private/OculusAudioGeometryComponent.cpp
@@ -70,10 +70,10 @@ void UOculusAudioGeometryComponent::TickComponent(float DeltaTime, enum ELevelTi
 		
 		float TransformMatrix[16] =
 		{
-			Right.Y,	Right.Z,	-Right.X,	0.0f,
-			Up.Y,	   Up.Z,	   -Up.X,	   0.0f,
-			Backward.Y, Backward.Z, -Backward.X, 0.0f,
-			Position.Y, Position.Z, -Position.X, 0.0f,
+			(float)Right.Y,	(float)Right.Z,	(float)-Right.X,	0.0f,
+			(float)Up.Y,	   (float)Up.Z,	   (float)-Up.X,	   0.0f,
+			(float)Backward.Y, (float)Backward.Z, (float)-Backward.X, 0.0f,
+			(float)Position.Y, (float)Position.Z, (float)-Position.X, 0.0f,
 		};
 		
 		ovrResult Result = OVRA_CALL(ovrAudio_AudioGeometrySetTransform)(ovrGeometry, TransformMatrix);
