Fix UE4 references in Oculus plugin (UE4 -> UE).

diff --git a/Engine/Plugins/Runtime/Oculus/OculusAudio/Source/OculusAudio/Private/OculusAudioContextManager.cpp b/Engine/Plugins/Runtime/Oculus/OculusAudio/Source/OculusAudio/Private/OculusAudioContextManager.cpp
--- a/Engine/Plugins/Runtime/Oculus/OculusAudio/Source/OculusAudio/Private/OculusAudioContextManager.cpp
+++ b/Engine/Plugins/Runtime/Oculus/OculusAudio/Source/OculusAudio/Private/OculusAudioContextManager.cpp
@@ -66,7 +66,7 @@ ovrAudioContext FOculusAudioContextManager::GetOrCreateSerializationContext(UAct
 	/* NOTE: this was meant to be some magic that would clean up the serialization 
 	context so there is no stale state in PIE. This failed when there was no geometry 
 	in the scene. Since the plugins are force to use one global context that seems 
-	like a better option for the vanilla UE4 integraiton to do as well for consistency.
+	like a better option for the vanilla UE integraiton to do as well for consistency.
 	*/
 	if (SerializationParent != nullptr)
 	{
diff --git a/Engine/Plugins/Runtime/Oculus/OculusAudio/Source/OculusAudio/Private/OculusAudioDllManager.cpp b/Engine/Plugins/Runtime/Oculus/OculusAudio/Source/OculusAudio/Private/OculusAudioDllManager.cpp
--- a/Engine/Plugins/Runtime/Oculus/OculusAudio/Source/OculusAudio/Private/OculusAudioDllManager.cpp
+++ b/Engine/Plugins/Runtime/Oculus/OculusAudio/Source/OculusAudio/Private/OculusAudioDllManager.cpp
@@ -90,18 +90,18 @@ bool FOculusAudioLibraryManager::LoadDll()
 	{
 		const TCHAR* WWISE_DLL_NAME = TEXT("OculusSpatializerWwise");
 		const TCHAR* FMOD_DLL_NAME = TEXT("OculusSpatializerFMOD");
-		const TCHAR* UE4_DLL_NAME = (sizeof(void*) == 4) ? TEXT("ovraudio32") : TEXT("ovraudio64");
+		const TCHAR* UE_DLL_NAME = (sizeof(void*) == 4) ? TEXT("ovraudio32") : TEXT("ovraudio64");
 
 #if PLATFORM_WINDOWS 
 
 #if WITH_EDITOR
 		const FString WwisePath = FPaths::ProjectDir() / FString::Printf(TEXT("Binaries/Win64/"));
 		const FString FMODPath = FPaths::ProjectPluginsDir() / FString::Printf(TEXT("FMODStudio/Binaries/Win64/"));
-		const FString UE4Path = FPaths::EngineDir() / FString::Printf(TEXT("Binaries/ThirdParty/Oculus/Audio/Win64/"));
+		const FString UEPath = FPaths::EngineDir() / FString::Printf(TEXT("Binaries/ThirdParty/Oculus/Audio/Win64/"));
 #else
 		const FString WwisePath = FPaths::ProjectDir() / FString::Printf(TEXT("Binaries/Win64/"));
 		const FString FMODPath = FPaths::ProjectPluginsDir() / FString::Printf(TEXT("FMODStudio/Binaries/Win64/")); //TODO verify this
-		const FString UE4Path = FPaths::ProjectDir() / FString::Printf(TEXT("../Engine/Binaries/ThirdParty/Oculus/Audio/Win64/"));
+		const FString UEPath = FPaths::ProjectDir() / FString::Printf(TEXT("../Engine/Binaries/ThirdParty/Oculus/Audio/Win64/"));
 #endif
 
 		FString Path;
@@ -111,21 +111,21 @@ bool FOculusAudioLibraryManager::LoadDll()
 			ClientType = OVRA_CLIENT_TYPE_WWISE_UNKNOWN;
 			Path = WwisePath;
 			DLL_NAME = WWISE_DLL_NAME;
-			UE_LOG(LogAudio, Display, TEXT("Oculus Audio: OculusSpatializerWwise.dll found, using the Wwise version of the Oculus Audio UE4 integration"));
+			UE_LOG(LogAudio, Display, TEXT("Oculus Audio: OculusSpatializerWwise.dll found, using the Wwise version of the Oculus Audio UE integration"));
 		}
 		else if (FPaths::FileExists(FMODPath + FMOD_DLL_NAME + ".dll"))
 		{
 			ClientType = OVRA_CLIENT_TYPE_FMOD;
 			Path = FMODPath;
 			DLL_NAME = FMOD_DLL_NAME;
-			UE_LOG(LogAudio, Display, TEXT("Oculus Audio: OculusSpatializerFMOD.dll found, using the FMOD version of the Oculus Audio UE4 integration"));
+			UE_LOG(LogAudio, Display, TEXT("Oculus Audio: OculusSpatializerFMOD.dll found, using the FMOD version of the Oculus Audio UE integration"));
 		}
 		else
 		{
 			ClientType = -1;
-			Path = UE4Path;
-			DLL_NAME = UE4_DLL_NAME;
-			UE_LOG(LogAudio, Display, TEXT("Oculus Audio: Middleware plugins not found, assuming native UE4 AudioMixer"));
+			Path = UEPath;
+			DLL_NAME = UE_DLL_NAME;
+			UE_LOG(LogAudio, Display, TEXT("Oculus Audio: Middleware plugins not found, assuming native UE AudioMixer"));
 		}
 
 
@@ -140,24 +140,24 @@ bool FOculusAudioLibraryManager::LoadDll()
 		OculusAudioDllHandle = FPlatformProcess::GetDllHandle(*(Path + WWISE_DLL_NAME + ".so"));
 		if (OculusAudioDllHandle != nullptr)
 		{
-			UE_LOG(LogAudio, Display, TEXT("Oculus Audio: %s found, using the Wwise version of the Oculus Audio UE4 integration"), WWISE_DLL_NAME);
+			UE_LOG(LogAudio, Display, TEXT("Oculus Audio: %s found, using the Wwise version of the Oculus Audio UE integration"), WWISE_DLL_NAME);
 			return true;
 		}
 		OculusAudioDllHandle = FPlatformProcess::GetDllHandle(*(Path + FMOD_DLL_NAME + ".so"));
 		if (OculusAudioDllHandle != nullptr)
 		{
-			UE_LOG(LogAudio, Display, TEXT("Oculus Audio: %s found, using the FMOD version of the Oculus Audio UE4 integration"), FMOD_DLL_NAME);
+			UE_LOG(LogAudio, Display, TEXT("Oculus Audio: %s found, using the FMOD version of the Oculus Audio UE integration"), FMOD_DLL_NAME);
 			return true;
 		}
-		OculusAudioDllHandle = FPlatformProcess::GetDllHandle(*(Path + UE4_DLL_NAME + ".so"));
+		OculusAudioDllHandle = FPlatformProcess::GetDllHandle(*(Path + UE_DLL_NAME + ".so"));
 		if (OculusAudioDllHandle != nullptr)
 		{
-			UE_LOG(LogAudio, Display, TEXT("Oculus Audio: Middleware plugins not found, %s found, assuming native UE4 AudioMixer"), UE4_DLL_NAME);
+			UE_LOG(LogAudio, Display, TEXT("Oculus Audio: Middleware plugins not found, %s found, assuming native UE AudioMixer"), UE_DLL_NAME);
 			return true;
 		}
 		else
 		{
-			UE_LOG(LogAudio, Error, TEXT("Oculus Audio: Unable to load Oculus Audio UE4 integratiton"), UE4_DLL_NAME);
+			UE_LOG(LogAudio, Error, TEXT("Oculus Audio: Unable to load Oculus Audio UE integratiton"), UE_DLL_NAME);
 		}
 #endif
 
diff --git a/Engine/Plugins/Runtime/Oculus/OculusVR/Source/OculusEditor/Private/OculusToolWidget.cpp b/Engine/Plugins/Runtime/Oculus/OculusVR/Source/OculusEditor/Private/OculusToolWidget.cpp
--- a/Engine/Plugins/Runtime/Oculus/OculusVR/Source/OculusEditor/Private/OculusToolWidget.cpp
+++ b/Engine/Plugins/Runtime/Oculus/OculusVR/Source/OculusEditor/Private/OculusToolWidget.cpp
@@ -23,7 +23,7 @@
 #define LOCTEXT_NAMESPACE "OculusToolWidget"
 
 // Misc notes and known issues:
-// * I save after every change because UE4 wasn't prompting to save on exit, but this makes it tough for users to undo, and doesn't prompt shader rebuild. Alternatives?
+// * I save after every change because UE wasn't prompting to save on exit, but this makes it tough for users to undo, and doesn't prompt shader rebuild. Alternatives?
 
 TSharedRef<SHorizontalBox> SOculusToolWidget::CreateSimpleSetting(SimpleSetting* setting)
 {
