Fix UE4 references in Oculus plugin.

diff --git a/Engine/Plugins/Runtime/Oculus/OculusAudio/Source/OculusAudio/Private/OculusAudioGeometryComponent.cpp b/Engine/Plugins/Runtime/Oculus/OculusAudio/Source/OculusAudio/Private/OculusAudioGeometryComponent.cpp
--- a/Engine/Plugins/Runtime/Oculus/OculusAudio/Source/OculusAudio/Private/OculusAudioGeometryComponent.cpp
+++ b/Engine/Plugins/Runtime/Oculus/OculusAudio/Source/OculusAudio/Private/OculusAudioGeometryComponent.cpp
@@ -61,7 +61,7 @@ void UOculusAudioGeometryComponent::TickComponent(float DeltaTime, enum ELevelTi
 	FTransform Transform = GetOwner()->GetTransform();
 	if (!Transform.Equals(PreviousTransform))
 	{
-		// UE4:	 x:forward, y:right, z:up
+		// UE:		x:forward, y:right, z:up
 		// Oculus:  x:right,   y:up,	z:backward
 		FVector Right = Transform.GetScaledAxis(EAxis::Y);
 		FVector Up = Transform.GetScaledAxis(EAxis::Z);
@@ -131,7 +131,7 @@ bool UOculusAudioGeometryComponent::UploadGeometry()
 
 #if 0 // TODO: map acoustic materials to visual/phyiscal materials
 	// The goal here is to build a hash table with unique graphic materials used by the acoustic geometry and automatically mapping them with
-	// our corresponding audio materials. A possible alternative is to do the mapping manually by adding a user data field in the UE4 material
+	// our corresponding audio materials. A possible alternative is to do the mapping manually by adding a user data field in the UE material
 	// definition and manually tagging each graphic material with an audio material in the editor. 
 	UStaticMeshComponent* StaticMeshComponent = SMActor->GetStaticMeshComponent();
 	for (int x = 0; x < StaticMeshComponent->GetNumMaterials(); ++x) {
@@ -322,7 +322,7 @@ ovrAudioContext UOculusAudioGeometryComponent::GetContext(UWorld* World)
 				FAudioDevice* AudioDevice = World->GetAudioDevice().GetAudioDevice();
 				if (AudioDevice == nullptr)
 				{
-					// This happens when cooking for native UE4 AudioMixer integration
+					// This happens when cooking for native UE AudioMixer integration
 					CachedContext = FOculusAudioContextManager::GetOrCreateSerializationContext(this);
 				}
 				else
diff --git a/Engine/Plugins/Runtime/Oculus/OculusAudio/Source/OculusAudio/Private/OculusAudioMixer.cpp b/Engine/Plugins/Runtime/Oculus/OculusAudio/Source/OculusAudio/Private/OculusAudioMixer.cpp
--- a/Engine/Plugins/Runtime/Oculus/OculusAudio/Source/OculusAudio/Private/OculusAudioMixer.cpp
+++ b/Engine/Plugins/Runtime/Oculus/OculusAudio/Source/OculusAudio/Private/OculusAudioMixer.cpp
@@ -41,7 +41,7 @@ void OculusAudioSpatializationAudioMixer::Initialize(const FAudioPluginInitializ
 
 	Params.AddDefaulted(InitParams.NumSources);
 
-	// UE4 is in centimeters
+	// UE is in centimeters
 	ovrResult Result = OVRA_CALL(ovrAudio_SetUnitScale)(Context, 0.01f);
 	OVR_AUDIO_CHECK(Result == ovrSuccess, "Failed to set unit scale");
 
