[Backout] - CL31402046
Original CL Desc
-----------------------------------------------------------------
Removing OculusAudio. Proper way to get the plugin is directly through vendor.

diff --git a/Engine/Plugins/Runtime/Oculus/OculusAudio/OculusAudio.uplugin b/Engine/Plugins/Runtime/Oculus/OculusAudio/OculusAudio.uplugin
new file mode 100644
--- /dev/null
+++ b/Engine/Plugins/Runtime/Oculus/OculusAudio/OculusAudio.uplugin
@@ -0,0 +1,36 @@
+{
+	"FileVersion": 3,
+	"Version": 1,
+	"VersionName": "1.34.0",
+	"FriendlyName": "Oculus Audio",
+	"Description": "Oculus audio spatialization",
+	"Category": "Audio",
+	"CreatedBy": "Oculus",
+	"CreatedByURL": "https://www.oculus.com/",
+	"DocsURL": "https://developer.oculus.com/documentation/unreal/latest/concepts/unreal-audio/",
+	"MarketplaceURL": "",
+	"SupportURL": "https://forums.oculusvr.com/developer/categories/audio-development",
+	"EnabledByDefault": false,
+	"CanContainContent": true,
+	"IsBetaVersion": false,
+	"Installed": false,
+	"Modules": [
+		{
+			"Name": "OculusAudio",
+			"Type": "Runtime",
+			"LoadingPhase": "PostConfigInit",
+			"PlatformAllowList": [
+				"Win64",
+				"Android"
+			]
+		},
+		{
+			"Name": "OculusAudioEditor",
+			"Type": "Editor",
+			"LoadingPhase": "PostEngineInit",
+			"PlatformAllowList": [
+				"Win64"
+			]
+		}
+	]
+}
\ No newline at end of file
diff --git a/Engine/Plugins/Runtime/Oculus/OculusAudio/Source/OculusAudio/OculusAudio.Build.cs b/Engine/Plugins/Runtime/Oculus/OculusAudio/Source/OculusAudio/OculusAudio.Build.cs
new file mode 100644
--- /dev/null
+++ b/Engine/Plugins/Runtime/Oculus/OculusAudio/Source/OculusAudio/OculusAudio.Build.cs
@@ -0,0 +1,54 @@
+// Copyright Epic Games, Inc. All Rights Reserved.
+
+using System.IO;
+
+namespace UnrealBuildTool.Rules
+{
+	public class OculusAudio : ModuleRules
+	{
+		public OculusAudio(ReadOnlyTargetRules Target) : base(Target)
+		{
+			PrivateIncludePaths.AddRange(
+				new string[] {
+					System.IO.Path.Combine(GetModuleDirectory("AudioMixer"), "Private"),
+
+					// ... add other private include paths required here ...
+				}
+				);
+
+			PrivateDependencyModuleNames.AddRange(
+				new string[]
+				{
+					"Core",
+					"CoreUObject",
+					"Engine",
+					"AudioMixer",
+					"Landscape",
+					"SoundFieldRendering",
+					"LibOVRAudio",
+					"SignalProcessing",
+					"AudioExtensions"
+				}
+				);
+
+			if (Target.Platform == UnrealTargetPlatform.Win64)
+			{
+				// Automatically copy DLL to packaged builds
+				RuntimeDependencies.Add("$(EngineDir)/Binaries/ThirdParty/Oculus/Audio/Win64/ovraudio64.dll");
+
+				AddEngineThirdPartyPrivateStaticDependencies(Target, "DX11Audio");
+			}
+
+			if (Target.Platform == UnrealTargetPlatform.Android)
+			{
+				// AndroidPlugin
+				{
+					string PluginPath = Utils.MakePathRelativeTo(ModuleDirectory, Target.RelativeEnginePath);
+					AdditionalPropertiesForReceipt.Add("AndroidPlugin", Path.Combine(PluginPath, "OculusAudio_APL.xml"));
+
+                    PublicAdditionalLibraries.Add("ThirdParty/Oculus/LibOVRAudio/LibOVRAudio/lib/arm64-v8a/libovraudio64.so");
+                }
+			}
+		}
+	}
+}
diff --git a/Engine/Plugins/Runtime/Oculus/OculusAudio/Source/OculusAudio/OculusAudio_APL.xml b/Engine/Plugins/Runtime/Oculus/OculusAudio/Source/OculusAudio/OculusAudio_APL.xml
new file mode 100644
--- /dev/null
+++ b/Engine/Plugins/Runtime/Oculus/OculusAudio/Source/OculusAudio/OculusAudio_APL.xml
@@ -0,0 +1,45 @@
+<?xml version="1.0" encoding="utf-8"?>
+<!-- steps to add to build additions -->
+<root xmlns:android="http://schemas.android.com/apk/res/android">
+	<!-- init section is always evaluated once per architecture -->
+	<init>
+		<log text="OculusAudio init"/>
+		<setBool result="bSupported" value="false"/>
+		<isArch arch="armeabi-v7a">
+			<setBool result="bSupported" value="true"/>
+		</isArch>
+		<isArch arch="arm64-v8a">
+			<setBool result="bSupported" value="true"/>
+		</isArch>
+	</init>
+
+	<!-- optional files or directories to copy to Intermediate/Android/APK -->
+	<resourceCopies>
+		<isArch arch="armeabi-v7a">
+			<!-- TODO: make this conditional and only copy the required verion!
+			<log text="Copying libOculusSpatializerWwise.so"/>
+			<copyFile src="$S(EngineDir)/Source/ThirdParty/Oculus/LibOVRAudio/LibOVRAudio/lib/armeabi-v7a/libOculusSpatializerWwise.so"
+							  dst="$S(BuildDir)/libs/armeabi-v7a/libOculusSpatializerWwise.so" />
+			<log text="Copying libOculusSpatializerFMOD.so"/>
+			<copyFile src="$S(EngineDir)/Source/ThirdParty/Oculus/LibOVRAudio/LibOVRAudio/lib/armeabi-v7a/libOculusSpatializerFMOD.so"
+							  dst="$S(BuildDir)/libs/armeabi-v7a/libOculusSpatializerFMOD.so" />
+			-->
+			<log text="Copying libovraudio32.so"/>
+			<copyFile src="$S(EngineDir)/Source/ThirdParty/Oculus/LibOVRAudio/LibOVRAudio/lib/armeabi-v7a/libovraudio32.so"
+							  dst="$S(BuildDir)/libs/armeabi-v7a/libovraudio32.so" />
+		</isArch>
+		<isArch arch="arm64-v8a">
+			<!-- TODO: make this conditional and only copy the required verion!
+			<log text="Copying libOculusSpatializerWwise.so"/>
+			<copyFile src="$S(EngineDir)/Source/ThirdParty/Oculus/LibOVRAudio/LibOVRAudio/lib/arm64-v8a/libOculusSpatializerWwise.so"
+							  dst="$S(BuildDir)/libs/arm64-v8a/libOculusSpatializerWwise.so" />
+			<log text="Copying libOculusSpatializerFMOD.so"/>
+			<copyFile src="$S(EngineDir)/Source/ThirdParty/Oculus/LibOVRAudio/LibOVRAudio/lib/arm64-v8a/libOculusSpatializerFMOD.so"
+							  dst="$S(BuildDir)/libs/arm64-v8a/libOculusSpatializerFMOD.so" />
+			-->
+			<log text="Copying libovraudio64.so"/>
+			<copyFile src="$S(EngineDir)/Source/ThirdParty/Oculus/LibOVRAudio/LibOVRAudio/lib/arm64-v8a/libovraudio64.so"
+							  dst="$S(BuildDir)/libs/arm64-v8a/libovraudio64.so" />
+		</isArch>
+	</resourceCopies>
+</root>
diff --git a/Engine/Plugins/Runtime/Oculus/OculusAudio/Source/OculusAudio/Private/IOculusAudioPlugin.cpp b/Engine/Plugins/Runtime/Oculus/OculusAudio/Source/OculusAudio/Private/IOculusAudioPlugin.cpp
new file mode 100644
--- /dev/null
+++ b/Engine/Plugins/Runtime/Oculus/OculusAudio/Source/OculusAudio/Private/IOculusAudioPlugin.cpp
@@ -0,0 +1,35 @@
+// Copyright Epic Games, Inc. All Rights Reserved.
+
+#include "IOculusAudioPlugin.h"
+#include "OculusAudioContextManager.h"
+#include "AudioDevice.h"
+#include "Features/IModularFeatures.h"
+
+void FOculusAudioPlugin::StartupModule()
+{
+	FOculusAudioLibraryManager::Get().Initialize();
+	IModularFeatures::Get().RegisterModularFeature(FOculusSpatializationPluginFactory::GetModularFeatureName(), &PluginFactory);
+	IModularFeatures::Get().RegisterModularFeature(FOculusReverbPluginFactory::GetModularFeatureName(), &ReverbPluginFactory);
+};
+
+void FOculusAudioPlugin::ShutdownModule()
+{
+	FOculusAudioLibraryManager::Get().Shutdown();
+}
+
+void FOculusAudioPlugin::RegisterAudioDevice(FAudioDevice* AudioDeviceHandle)
+{
+	if (!RegisteredAudioDevices.Contains(AudioDeviceHandle))
+	{
+		TAudioPluginListenerPtr ContextManager = TAudioPluginListenerPtr(new FOculusAudioContextManager());
+		AudioDeviceHandle->RegisterPluginListener(ContextManager);
+		RegisteredAudioDevices.Add(AudioDeviceHandle);
+	}
+}
+
+void FOculusAudioPlugin::UnregisterAudioDevice(FAudioDevice* AudioDeviceHandle)
+{
+	RegisteredAudioDevices.Remove(AudioDeviceHandle);
+}
+
+IMPLEMENT_MODULE(FOculusAudioPlugin, OculusAudio)
\ No newline at end of file
diff --git a/Engine/Plugins/Runtime/Oculus/OculusAudio/Source/OculusAudio/Private/OculusAmbisonicSpatializer.cpp b/Engine/Plugins/Runtime/Oculus/OculusAudio/Source/OculusAudio/Private/OculusAmbisonicSpatializer.cpp
new file mode 100644
--- /dev/null
+++ b/Engine/Plugins/Runtime/Oculus/OculusAudio/Source/OculusAudio/Private/OculusAmbisonicSpatializer.cpp
@@ -0,0 +1,417 @@
+// Copyright Epic Games, Inc. All Rights Reserved.
+
+#include "OculusAmbisonicSpatializer.h"
+#include "ISoundfieldFormat.h"
+#include "OculusAudioMixer.h"
+#include "DSP/FloatArrayMath.h"
+#include "AudioMixerDevice.h"
+#include "OculusAudioDllManager.h"
+#include "OculusAudioContextManager.h"
+#include "SoundFieldRendering.h"
+#include "AudioPluginUtilities.h"
+
+namespace OculusAudioUtils
+{
+	bool IsOculusAudioTheCurrentSpatializationPlugin()
+	{
+#if WITH_EDITOR
+		static FString OculusDisplayName = FString(TEXT("Oculus Audio"));
+		return AudioPluginUtilities::GetDesiredPluginName(EAudioPlugin::SPATIALIZATION).Equals(OculusDisplayName);
+#else
+		// For non editor situations, we can cache whether this is the current plugin or not the first time we check.
+		static FString OculusDisplayName = FString(TEXT("Oculus Audio"));
+		static bool bCheckedSpatializationPlugin = false;
+		static bool bIsOculusCurrentSpatiatizationPlugin = false;
+
+		if (!bCheckedSpatializationPlugin)
+		{
+			bIsOculusCurrentSpatiatizationPlugin = AudioPluginUtilities::GetDesiredPluginName(EAudioPlugin::SPATIALIZATION).Equals(OculusDisplayName);
+			bCheckedSpatializationPlugin = true;
+		}
+
+		return bIsOculusCurrentSpatiatizationPlugin;
+#endif
+	}
+}
+
+// the Oculus Audio renderer renders all inputs immediately to an interleaved stereo output.
+class FOculusSoundfieldBuffer : public ISoundfieldAudioPacket
+{
+public:
+	// Interleaved binaural audio buffer.
+	Audio::FAlignedFloatBuffer AudioBuffer;
+	int32 NumChannels;
+
+	FOculusSoundfieldBuffer()
+		: NumChannels(0)
+	{}
+
+	virtual void Serialize(FArchive& Ar) override
+	{
+		Ar << AudioBuffer;
+		Ar << NumChannels;
+	}
+
+
+	virtual TUniquePtr<ISoundfieldAudioPacket> Duplicate() const override
+	{
+		return TUniquePtr<ISoundfieldAudioPacket>(new FOculusSoundfieldBuffer(*this));
+	}
+
+
+	virtual void Reset() override
+	{
+		AudioBuffer.Reset();
+		NumChannels = 0;
+	}
+
+};
+
+class FOculusEncoder : public ISoundfieldEncoderStream
+{
+private:
+	ovrAudioContext Context;
+	TArray<int32> SourceIds;
+
+	// temp buffers use to binauralize each source.
+	Audio::FAlignedFloatBuffer ScratchMonoBuffer;
+	Audio::FAlignedFloatBuffer ScratchOutputBuffer;
+
+	static int32 SourceIdCounter;
+	static FCriticalSection SourceIdCounterCritSection;
+
+public:
+
+	FOculusEncoder(ovrAudioContext InContext, int32 InNumChannels, int32 MaxNumSources)
+		: Context(InContext)
+	{
+		check(Context != nullptr);
+		{
+			// Assign source Ids for each channel here.
+			FScopeLock ScopeLock(&SourceIdCounterCritSection);
+			for (int32 Index = 0; Index < InNumChannels; Index++)
+			{
+
+				SourceIds.Add(SourceIdCounter);
+				SourceIdCounter = (SourceIdCounter + 1) % MaxNumSources;
+			}
+		}
+
+		// Set default settings for each source channel here.
+		for (int32& SourceId : SourceIds)
+		{
+			uint32 Flags = 0;
+			Flags |= ovrAudioSourceFlag_ReflectionsDisabled;
+
+			ovrResult Result = OVRA_CALL(ovrAudio_SetAudioSourceFlags)(Context, SourceId, Flags);
+			check(Result == ovrSuccess);
+
+			ovrAudioSourceAttenuationMode mode = ovrAudioSourceAttenuationMode_None;
+			Result = OVRA_CALL(ovrAudio_SetAudioSourceAttenuationMode)(Context, SourceId, mode, 1.0f);
+			check(Result == ovrSuccess);
+
+			Result = OVRA_CALL(ovrAudio_SetAudioReverbSendLevel)(Context, SourceId, 0.0f);
+			check(Result == ovrSuccess);
+		}
+	}
+
+	virtual void Encode(const FSoundfieldEncoderInputData& InputData, ISoundfieldAudioPacket& OutputData) override
+	{
+		OutputData.Reset();
+		EncodeAndMixIn(InputData, OutputData);
+	}
+
+
+	// Binaurally spatialize each independent speaker position. If the spatialization plugin for the current platform was set to Oculus Audio
+	// when this encoder was created this will no-op.
+	virtual void EncodeAndMixIn(const FSoundfieldEncoderInputData& InputData, ISoundfieldAudioPacket& OutputData) override
+	{
+		FOculusSoundfieldBuffer& OutputBuffer = DowncastSoundfieldRef<FOculusSoundfieldBuffer>(OutputData);
+
+		check(InputData.NumChannels == SourceIds.Num());
+		const int32 NumFrames = InputData.AudioBuffer.Num() / InputData.NumChannels;
+
+		OutputBuffer.AudioBuffer.Reset();
+		OutputBuffer.NumChannels = 2;
+		OutputBuffer.AudioBuffer.AddZeroed(NumFrames * OutputBuffer.NumChannels);
+
+		FQuat ListenerOrientation = InputData.PositionalData.Rotation.Inverse();
+
+		// Translate the input position to OVR coordinates
+		FVector ListenerPosition = OculusAudioSpatializationAudioMixer::ToOVRVector(FVector::ZeroVector);
+		FVector ListenerForward = OculusAudioSpatializationAudioMixer::ToOVRVector(ListenerOrientation.GetForwardVector());
+		FVector OvrListenerUp = OculusAudioSpatializationAudioMixer::ToOVRVector(ListenerOrientation.GetUpVector());
+
+		ovrResult Result = OVRA_CALL(ovrAudio_SetListenerVectors)(Context,
+			ListenerPosition.X, ListenerPosition.Y, ListenerPosition.Z,
+			ListenerForward.X, ListenerForward.Y, ListenerForward.Z,
+			OvrListenerUp.X, OvrListenerUp.Y, OvrListenerUp.Z);
+		OVR_AUDIO_CHECK(Result, "Failed to set listener position and rotation");
+
+		check(InputData.PositionalData.ChannelPositions);
+		const TArray<Audio::FChannelPositionInfo>& ChannelPositions = *(InputData.PositionalData.ChannelPositions);
+
+		for (int32 ChannelIndex = 0; ChannelIndex < SourceIds.Num(); ChannelIndex++)
+		{
+			int32& SourceId = SourceIds[ChannelIndex];
+			const Audio::FChannelPositionInfo& ChannelPosition = ChannelPositions[ChannelIndex];
+			
+
+			// Translate the input position to OVR coordinates
+			FVector SourcePosition = OculusAudioSpatializationAudioMixer::ToOVRVector(ChannelPosition);
+
+			// Set the source position to current audio position
+			Result = OVRA_CALL(ovrAudio_SetAudioSourcePos)(Context, SourceId, SourcePosition.X, SourcePosition.Y, SourcePosition.Z);
+			OVR_AUDIO_CHECK(Result, "Failed to set audio source position");
+
+			// Deinterleave the audio into the mono temp buffer.
+			ScratchMonoBuffer.Reset();
+			ScratchMonoBuffer.AddUninitialized(NumFrames);
+			for(int32 FrameIndex = 0; FrameIndex < NumFrames; FrameIndex++)
+			{
+				ScratchMonoBuffer[FrameIndex] = InputData.AudioBuffer[FrameIndex * InputData.NumChannels + ChannelIndex];
+			}
+
+			ScratchOutputBuffer.Reset();
+			ScratchOutputBuffer.AddUninitialized(NumFrames * 2);
+
+			uint32 Status = 0;
+			Result = OVRA_CALL(ovrAudio_SpatializeMonoSourceInterleaved)(Context, SourceId, &Status, ScratchOutputBuffer.GetData(), ScratchMonoBuffer.GetData());
+			OVR_AUDIO_CHECK(Result, "Failed to spatialize mono source interleaved");
+
+			// sum the interleaved output into the output bed.
+			Audio::ArrayMixIn(ScratchOutputBuffer, OutputBuffer.AudioBuffer);
+		}
+	}
+};
+
+int32 FOculusEncoder::SourceIdCounter = 0;
+FCriticalSection FOculusEncoder::SourceIdCounterCritSection;
+
+// Because Oculus spatializes every directly to an interleaved stereo buffer,
+// we simply pass it forward here.
+class FOculusAudioDecoder : public ISoundfieldDecoderStream
+{
+private:
+	const bool bShouldUseSubmixForOutput;
+
+public:
+	FOculusAudioDecoder(const bool bInShouldUseSubmixForOutput)
+		: bShouldUseSubmixForOutput(bInShouldUseSubmixForOutput)
+	{}
+
+	FOculusAudioDecoder() = delete;
+
+	virtual void Decode(const FSoundfieldDecoderInputData& InputData, FSoundfieldDecoderOutputData& OutputData) override
+	{
+		OutputData.AudioBuffer.Reset();
+		OutputData.AudioBuffer.AddZeroed(InputData.NumFrames * InputData.PositionalData.NumChannels);
+		DecodeAndMixIn(InputData, OutputData);
+	}
+
+	virtual void DecodeAndMixIn(const FSoundfieldDecoderInputData& InputData, FSoundfieldDecoderOutputData& OutputData) override
+	{
+		if (bShouldUseSubmixForOutput)
+		{
+			const FOculusSoundfieldBuffer& InputBuffer = DowncastSoundfieldRef<const FOculusSoundfieldBuffer>(InputData.SoundfieldBuffer);
+
+			if (InputBuffer.AudioBuffer.Num() == 0 || InputBuffer.NumChannels == 0)
+			{
+				return;
+			}
+			else if (InputData.PositionalData.NumChannels == InputBuffer.NumChannels)
+			{
+				// If the number of output channels is the same as the input channels, mix it in directly.
+				Audio::ArrayMixIn(InputBuffer.AudioBuffer, OutputData.AudioBuffer);
+			}
+			else
+			{
+				// Otherwise, downmix and mix in.
+				Audio::FAlignedFloatBuffer OutChannelMap;
+				Audio::FMixerDevice::Get2DChannelMap(false, InputBuffer.NumChannels, InputData.PositionalData.NumChannels, false, OutChannelMap);
+				Audio::DownmixAndSumIntoBuffer(InputBuffer.NumChannels, InputData.PositionalData.NumChannels, InputBuffer.AudioBuffer, OutputData.AudioBuffer, OutChannelMap.GetData());
+			}
+		}
+	}
+};
+
+// Class that takes ambisonics audio and renders it to an FOculusSoundfieldBuffer.
+class FOculusAmbisonicsTranscoder : public ISoundfieldTranscodeStream
+{
+private:
+	ovrAudioContext Context;
+	ovrAudioAmbisonicStream AmbiStreamObject;
+public:
+	FOculusAmbisonicsTranscoder(ovrAudioContext InContext, int32 InSampleRate, int32 InBufferLength)
+		: Context(InContext)
+	{
+		ovrAudioAmbisonicFormat StreamFormat = ovrAudioAmbisonicFormat_AmbiX;
+		ovrResult OurResult = OVRA_CALL(ovrAudio_CreateAmbisonicStream)(Context, InSampleRate, InBufferLength, StreamFormat, 1, &AmbiStreamObject);
+		check(OurResult == 0);
+		check(AmbiStreamObject != nullptr);
+	}
+
+	virtual void Transcode(const ISoundfieldAudioPacket& InputData, const ISoundfieldEncodingSettingsProxy& InputSettings, ISoundfieldAudioPacket& OutputData, const ISoundfieldEncodingSettingsProxy& OutputSettings) override
+	{
+		OutputData.Reset();
+		TranscodeAndMixIn(InputData, InputSettings, OutputData, OutputSettings);
+	}
+
+
+	virtual void TranscodeAndMixIn(const ISoundfieldAudioPacket& InputData, const ISoundfieldEncodingSettingsProxy& InputSettings, ISoundfieldAudioPacket& PacketToSumTo, const ISoundfieldEncodingSettingsProxy& OutputSettings) override
+	{
+		const FAmbisonicsSoundfieldBuffer& InputBuffer = DowncastSoundfieldRef<const FAmbisonicsSoundfieldBuffer>(InputData);
+		FOculusSoundfieldBuffer& OutputBuffer = DowncastSoundfieldRef<FOculusSoundfieldBuffer>(PacketToSumTo);
+
+		if (!InputBuffer.NumChannels || InputBuffer.AudioBuffer.IsEmpty())
+		{
+			return;
+		}
+
+		const int32 NumFrames = InputBuffer.AudioBuffer.Num() / InputBuffer.NumChannels;
+		OutputBuffer.AudioBuffer.SetNumZeroed(NumFrames * 2);
+		OutputBuffer.NumChannels = 2;
+
+		//Currently, Oculus only decodes first-order ambisonics to stereo.
+		// in the future we can truncate InputBuffer to four channels and 
+		// render that.
+		// For now we just output silence for higher order ambisonics from the Oculus spatializaers
+		if (InputBuffer.NumChannels == 4)
+		{
+			// Translate the input position to OVR coordinates
+			const FQuat& ListenerRotation = InputBuffer.Rotation;
+			FVector OvrListenerForward = OculusAudioSpatializationAudioMixer::ToOVRVector(ListenerRotation.GetForwardVector());
+			FVector OvrListenerUp = OculusAudioSpatializationAudioMixer::ToOVRVector(ListenerRotation.GetUpVector());
+
+			static auto SetListenerVectors = OVRA_CALL(ovrAudio_SetListenerVectors);
+			ovrResult Result = SetListenerVectors(Context,
+				0.0f, 0.0f, 0.0f,
+				OvrListenerForward.X, OvrListenerForward.Y, OvrListenerForward.Z,
+				OvrListenerUp.X, OvrListenerUp.Y, OvrListenerUp.Z);
+
+			ovrResult DecodeResult = OVRA_CALL(ovrAudio_ProcessAmbisonicStreamInterleaved)(Context, AmbiStreamObject, InputBuffer.AudioBuffer.GetData(), OutputBuffer.AudioBuffer.GetData(), NumFrames);
+			check(DecodeResult == 0);
+		}
+	}
+
+};
+
+// Since FOculusSoundfieldBuffer is just an interleaved stereo buffer, we simply mix the buffer here.
+class FOculusAudioSoundfieldMixer : public ISoundfieldMixerStream
+{
+public:
+	virtual void MixTogether(const FSoundfieldMixerInputData& InputData, ISoundfieldAudioPacket& PacketToSumTo) override
+	{
+		const FOculusSoundfieldBuffer& InputBuffer = DowncastSoundfieldRef<const FOculusSoundfieldBuffer>(InputData.InputPacket);
+		FOculusSoundfieldBuffer& OutputBuffer = DowncastSoundfieldRef<FOculusSoundfieldBuffer>(PacketToSumTo);
+
+		if (OutputBuffer.AudioBuffer.Num() == 0)
+		{
+			OutputBuffer.AudioBuffer = InputBuffer.AudioBuffer;
+			OutputBuffer.NumChannels = InputBuffer.NumChannels;
+		}
+		else
+		{
+			check(InputBuffer.AudioBuffer.Num() == OutputBuffer.AudioBuffer.Num());
+			check(InputBuffer.NumChannels == OutputBuffer.NumChannels);
+
+
+			Audio::ArrayMixIn(InputBuffer.AudioBuffer, OutputBuffer.AudioBuffer, InputData.SendLevel);
+		}
+	}
+};
+
+FOculusAmbisonicsFactory::FOculusAmbisonicsFactory()
+{
+	ISoundfieldFactory::RegisterSoundfieldFormat(this);
+}
+
+FOculusAmbisonicsFactory::~FOculusAmbisonicsFactory()
+{
+	ISoundfieldFactory::UnregisterSoundfieldFormat(this);
+}
+
+bool FOculusAmbisonicsFactory::IsTranscodeRequiredBetweenSettings(const ISoundfieldEncodingSettingsProxy& SourceSettings, const ISoundfieldEncodingSettingsProxy& DestinationSettings)
+{
+	return false;
+}
+
+bool FOculusAmbisonicsFactory::CanTranscodeFromSoundfieldFormat(FName SourceFormat, const ISoundfieldEncodingSettingsProxy& SourceEncodingSettings)
+{
+	FName AmbisonicsSoundfieldFormatName = GetUnrealAmbisonicsFormatName();
+	if (SourceFormat == AmbisonicsSoundfieldFormatName)
+	{
+		return true;
+	}
+	else
+	{
+		return false;
+	}
+}
+
+bool FOculusAmbisonicsFactory::CanTranscodeToSoundfieldFormat(FName DestinationFormat, const ISoundfieldEncodingSettingsProxy& DestinationEncodingSettings)
+{
+	return false;
+}
+
+UClass* FOculusAmbisonicsFactory::GetCustomEncodingSettingsClass() const
+{
+	return UOculusAudioSoundfieldSettings::StaticClass();
+}
+
+const USoundfieldEncodingSettingsBase* FOculusAmbisonicsFactory::GetDefaultEncodingSettings()
+{
+	return GetDefault<UOculusAudioSoundfieldSettings>();
+}
+
+FName FOculusAmbisonicsFactory::GetSoundfieldFormatName()
+{
+	static FName OculusBinauralFormatName = FName(TEXT("Oculus Binaural"));
+	return OculusBinauralFormatName;
+}
+
+TUniquePtr<ISoundfieldEncoderStream> FOculusAmbisonicsFactory::CreateEncoderStream(const FAudioPluginInitializationParams& InitInfo, const ISoundfieldEncodingSettingsProxy& InitialSettings)
+{
+	ovrAudioContext Context = FOculusAudioContextManager::GetContextForAudioDevice(InitInfo.AudioDevicePtr);
+	if (!Context)
+	{
+		Context = FOculusAudioContextManager::CreateContextForAudioDevice(InitInfo.AudioDevicePtr);
+		check(Context);
+	}
+
+	return TUniquePtr<ISoundfieldEncoderStream>(new FOculusEncoder(Context, InitInfo.NumOutputChannels, InitInfo.AudioDevicePtr->GetMaxChannels()));
+}
+
+TUniquePtr<ISoundfieldDecoderStream> FOculusAmbisonicsFactory::CreateDecoderStream(const FAudioPluginInitializationParams& InitInfo, const ISoundfieldEncodingSettingsProxy& InitialSettings)
+{
+	const bool bShouldOutputOculusAudio = !OculusAudioUtils::IsOculusAudioTheCurrentSpatializationPlugin();
+
+	UE_CLOG(!bShouldOutputOculusAudio, LogAudio, Warning, TEXT("Oculus Audio is the currently selected spatialization plugin. The Oculus soundfield submix's output will be ignored."));
+
+	return TUniquePtr<ISoundfieldDecoderStream>(new FOculusAudioDecoder(bShouldOutputOculusAudio));
+}
+
+TUniquePtr<ISoundfieldTranscodeStream> FOculusAmbisonicsFactory::CreateTranscoderStream(const FName SourceFormat, const ISoundfieldEncodingSettingsProxy& InitialSourceSettings, const FName DestinationFormat, const ISoundfieldEncodingSettingsProxy& InitialDestinationSettings, const FAudioPluginInitializationParams& InitInfo)
+{
+	check(SourceFormat == GetUnrealAmbisonicsFormatName());
+	check(DestinationFormat == GetSoundfieldFormatName());
+
+	ovrAudioContext Context = FOculusAudioContextManager::GetContextForAudioDevice(InitInfo.AudioDevicePtr);
+	if (!Context)
+	{
+		Context = FOculusAudioContextManager::CreateContextForAudioDevice(InitInfo.AudioDevicePtr);
+		check(Context);
+	}
+
+	return TUniquePtr<ISoundfieldTranscodeStream>(new FOculusAmbisonicsTranscoder(Context, InitInfo.SampleRate, InitInfo.BufferLength));
+}
+
+TUniquePtr<ISoundfieldMixerStream> FOculusAmbisonicsFactory::CreateMixerStream(const ISoundfieldEncodingSettingsProxy& InitialSettings)
+{
+	return TUniquePtr<ISoundfieldMixerStream>(new FOculusAudioSoundfieldMixer());
+}
+
+TUniquePtr<ISoundfieldAudioPacket> FOculusAmbisonicsFactory::CreateEmptyPacket()
+{
+	return TUniquePtr<ISoundfieldAudioPacket>(new FOculusSoundfieldBuffer());
+}
diff --git a/Engine/Plugins/Runtime/Oculus/OculusAudio/Source/OculusAudio/Private/OculusAudio.cpp b/Engine/Plugins/Runtime/Oculus/OculusAudio/Source/OculusAudio/Private/OculusAudio.cpp
new file mode 100644
--- /dev/null
+++ b/Engine/Plugins/Runtime/Oculus/OculusAudio/Source/OculusAudio/Private/OculusAudio.cpp
@@ -0,0 +1,26 @@
+// Copyright Epic Games, Inc. All Rights Reserved.
+
+#include "OculusAudio.h"
+#include "OculusAmbisonicSpatializer.h"
+#include "OculusAudioReverb.h"
+#include "IOculusAudioPlugin.h"
+#include "Features/IModularFeatures.h"
+
+TAudioSpatializationPtr FOculusSpatializationPluginFactory::CreateNewSpatializationPlugin(FAudioDevice* OwningDevice)
+{
+	FOculusAudioPlugin* Plugin = &FModuleManager::GetModuleChecked<FOculusAudioPlugin>("OculusAudio");
+	check(Plugin != nullptr);
+	Plugin->RegisterAudioDevice(OwningDevice);
+
+	return TAudioSpatializationPtr(new OculusAudioSpatializationAudioMixer());
+}
+
+
+TAudioReverbPtr FOculusReverbPluginFactory::CreateNewReverbPlugin(FAudioDevice* OwningDevice)
+{
+	FOculusAudioPlugin* Plugin = &FModuleManager::GetModuleChecked<FOculusAudioPlugin>("OculusAudio");
+	check(Plugin != nullptr);
+	Plugin->RegisterAudioDevice(OwningDevice);
+
+	return TAudioReverbPtr(new OculusAudioReverb());
+}
\ No newline at end of file
diff --git a/Engine/Plugins/Runtime/Oculus/OculusAudio/Source/OculusAudio/Private/OculusAudio.h b/Engine/Plugins/Runtime/Oculus/OculusAudio/Source/OculusAudio/Private/OculusAudio.h
new file mode 100644
--- /dev/null
+++ b/Engine/Plugins/Runtime/Oculus/OculusAudio/Source/OculusAudio/Private/OculusAudio.h
@@ -0,0 +1,55 @@
+// Copyright Epic Games, Inc. All Rights Reserved.
+
+#pragma once
+
+#include "IAudioExtensionPlugin.h"
+#include "OculusAudioMixer.h"
+#include "OculusAudioSourceSettings.h"
+#include "OculusAmbisonicSpatializer.h"
+
+/************************************************************************/
+/* FOculusSpatializationPluginFactory								   */
+/* Handles initialization of the required Oculus Audio Spatialization   */
+/* plugin.															  */
+/************************************************************************/
+class FOculusSpatializationPluginFactory : public IAudioSpatializationFactory
+{
+public:
+	//~ Begin IAudioSpatializationFactory
+	virtual FString GetDisplayName() override
+	{
+		static FString DisplayName = FString(TEXT("Oculus Audio"));
+		return DisplayName;
+	}
+
+	virtual bool SupportsPlatform(const FString& PlatformName) override
+	{
+		return	PlatformName == FString(TEXT("Windows")) || PlatformName == FString(TEXT("Android"));
+	}
+
+	virtual UClass* GetCustomSpatializationSettingsClass() const override { return UOculusAudioSourceSettings::StaticClass(); }
+
+	virtual TAudioSpatializationPtr CreateNewSpatializationPlugin(FAudioDevice* OwningDevice) override;
+	//~ End IAudioSpatializationFactory
+
+};
+
+class FOculusReverbPluginFactory : public IAudioReverbFactory
+{
+public:
+	//~ Begin IAudioReverbFactory
+	virtual FString GetDisplayName() override
+	{
+		static FString DisplayName = FString(TEXT("Oculus Audio"));
+		return DisplayName;
+	}
+
+	virtual bool SupportsPlatform(const FString& PlatformName) override
+	{
+		return	PlatformName == FString(TEXT("Windows")) || PlatformName == FString(TEXT("Android"));
+	}
+
+	virtual TAudioReverbPtr CreateNewReverbPlugin(FAudioDevice* OwningDevice) override;
+	//~ End IAudioReverbFactory
+};
+
diff --git a/Engine/Plugins/Runtime/Oculus/OculusAudio/Source/OculusAudio/Private/OculusAudioContextManager.cpp b/Engine/Plugins/Runtime/Oculus/OculusAudio/Source/OculusAudio/Private/OculusAudioContextManager.cpp
new file mode 100644
--- /dev/null
+++ b/Engine/Plugins/Runtime/Oculus/OculusAudio/Source/OculusAudio/Private/OculusAudioContextManager.cpp
@@ -0,0 +1,160 @@
+// Copyright Epic Games, Inc. All Rights Reserved.
+
+#include "OculusAudioContextManager.h"
+#include "OculusAmbisonicSpatializer.h"
+#include "OculusAudioReverb.h"
+#include "Features/IModularFeatures.h"
+#include "AudioMixerDevice.h"
+#include "IOculusAudioPlugin.h"
+
+ovrAudioContext FOculusAudioContextManager::SerializationContext = nullptr;
+UActorComponent* FOculusAudioContextManager::SerializationParent = nullptr;
+
+FCriticalSection FOculusAudioContextManager::SerializationContextLock;
+
+TMap<const FAudioDevice*, ovrAudioContext> FOculusAudioContextManager::ContextMap;
+FCriticalSection FOculusAudioContextManager::ContextCritSection;
+
+FOculusAudioContextManager::FOculusAudioContextManager()
+{
+	// empty
+}
+
+FOculusAudioContextManager::~FOculusAudioContextManager()
+{
+	SerializationContext = nullptr;
+}
+
+void FOculusAudioContextManager::OnListenerInitialize(FAudioDevice* AudioDevice, UWorld* ListenerWorld)
+{
+}
+
+void FOculusAudioContextManager::OnListenerShutdown(FAudioDevice* AudioDevice)
+{
+	FOculusAudioPlugin* Plugin = &FModuleManager::GetModuleChecked<FOculusAudioPlugin>("OculusAudio");
+	const FString OculusSpatializerPluginName = Plugin->GetSpatializationPluginFactory()->GetDisplayName();
+	const FString CurrentSpatializerPluginName = AudioPluginUtilities::GetDesiredPluginName(EAudioPlugin::SPATIALIZATION);
+	if (CurrentSpatializerPluginName.Equals(OculusSpatializerPluginName)) // we have a match!
+	{
+		OculusAudioSpatializationAudioMixer* Spatializer =
+			static_cast<OculusAudioSpatializationAudioMixer*>(AudioDevice->GetSpatializationPluginInterface().Get());
+
+		Spatializer->ClearContext();
+	}
+
+	const FString OculusReverbPluginName = Plugin->GetReverbPluginFactory()->GetDisplayName();
+	const FString CurrentReverbPluginName = AudioPluginUtilities::GetDesiredPluginName(EAudioPlugin::REVERB);
+	if (CurrentReverbPluginName.Equals(OculusReverbPluginName))
+	{
+		OculusAudioReverb* Reverb = static_cast<OculusAudioReverb*>(AudioDevice->ReverbPluginInterface.Get());
+		Reverb->ClearContext();
+	}
+
+	// FIXME:
+	// There's a possibility this will leak if a Oculus Binaural submix is created,
+	// but Oculus audio is not specified as the spatialization or reverb plugin.
+	// This is a niche use case, but could be solved by having the oculus soundfield explicitly destroy
+	// a context.
+	DestroyContextForAudioDevice(AudioDevice);
+	SerializationContext = nullptr;
+}
+
+ovrAudioContext FOculusAudioContextManager::GetOrCreateSerializationContext(UActorComponent* Parent)
+{
+	FScopeLock ScopeLock(&SerializationContextLock);
+
+#if 0 //WITH_EDITOR
+	/* NOTE: this was meant to be some magic that would clean up the serialization 
+	context so there is no stale state in PIE. This failed when there was no geometry 
+	in the scene. Since the plugins are force to use one global context that seems 
+	like a better option for the vanilla UE integraiton to do as well for consistency.
+	*/
+	if (SerializationParent != nullptr)
+	{
+		AActor* Actor = SerializationParent->GetOwner();
+		if (Actor != nullptr)
+		{
+			UWorld* World = Actor->GetWorld();
+			if (World != nullptr)
+			{
+				if (World->WorldType == EWorldType::Editor)
+				{
+					// It was just an editor serialize, we don't need these geometry objects
+					OVRA_CALL(ovrAudio_DestroyContext)(SerializationContext);
+					SerializationContext = nullptr;
+					SerializationParent = nullptr;
+				}
+			}
+		}
+	}
+#endif
+
+	ovrAudioContext PluginContext = FOculusAudioLibraryManager::Get().GetPluginContext();
+	if (PluginContext != nullptr)
+	{
+		return PluginContext;
+	}
+
+	if (SerializationContext == nullptr)
+	{
+		ovrResult Result = OVRA_CALL(ovrAudio_CreateContext)(&SerializationContext, nullptr);
+		if (Result != ovrSuccess)
+		{
+			const TCHAR* ErrString = GetOculusErrorString(Result);
+			UE_LOG(LogAudio, Error, TEXT("Oculus Audio SDK Error - %s: %s"), TEXT("Failed to create Oculus Audio context for serialization"), ErrString);
+			return nullptr;
+		}
+
+		SerializationParent = Parent;
+	}
+
+	return SerializationContext;
+}
+
+ovrAudioContext FOculusAudioContextManager::GetContextForAudioDevice(const FAudioDevice* InAudioDevice)
+{
+	FScopeLock ScopeLock(&ContextCritSection);
+	ovrAudioContext* Context = ContextMap.Find(InAudioDevice);
+	if (Context)
+	{
+		return *Context;
+	}
+	else
+	{
+		return nullptr;
+	}
+}
+
+ovrAudioContext FOculusAudioContextManager::CreateContextForAudioDevice(FAudioDevice* InAudioDevice)
+{
+	ovrAudioContextConfiguration ContextConfig = { 0 };
+	ContextConfig.acc_BufferLength = InAudioDevice->GetBufferLength();
+	ContextConfig.acc_MaxNumSources = InAudioDevice->GetMaxSources();
+	ContextConfig.acc_SampleRate = InAudioDevice->GetSampleRate();
+	ContextConfig.acc_Size = sizeof(ovrAudioContextConfiguration);
+
+	ovrAudioContext NewContext = nullptr;
+	ovrResult Result = OVRA_CALL(ovrAudio_CreateContext)(&NewContext, &ContextConfig);
+
+	if(ensure(Result == ovrSuccess))
+	{
+		FScopeLock ScopeLock(&ContextCritSection);
+		return ContextMap.Add(InAudioDevice, NewContext);
+	}
+	else
+	{
+		return nullptr;
+	}
+}
+
+void FOculusAudioContextManager::DestroyContextForAudioDevice(const FAudioDevice* InAudioDevice)
+{
+	FScopeLock ScopeLock(&ContextCritSection);
+	ovrAudioContext* Context = ContextMap.Find(InAudioDevice);
+
+	if (Context)
+	{
+		OVRA_CALL(ovrAudio_DestroyContext)(*Context);
+		ContextMap.Remove(InAudioDevice);
+	}
+}
diff --git a/Engine/Plugins/Runtime/Oculus/OculusAudio/Source/OculusAudio/Private/OculusAudioContextManager.h b/Engine/Plugins/Runtime/Oculus/OculusAudio/Source/OculusAudio/Private/OculusAudioContextManager.h
new file mode 100644
--- /dev/null
+++ b/Engine/Plugins/Runtime/Oculus/OculusAudio/Source/OculusAudio/Private/OculusAudioContextManager.h
@@ -0,0 +1,37 @@
+// Copyright Epic Games, Inc. All Rights Reserved.
+
+#pragma once
+
+#include "AudioPluginUtilities.h"
+#include "OVR_Audio.h"
+
+class UActorComponent;
+class FOculusAudioContextManager : public IAudioPluginListener
+{
+public:
+	FOculusAudioContextManager();
+	virtual ~FOculusAudioContextManager() override;
+
+	virtual void OnListenerInitialize(FAudioDevice* AudioDevice, UWorld* ListenerWorld) override;
+	virtual void OnListenerShutdown(FAudioDevice* AudioDevice) override;
+
+	static ovrAudioContext GetOrCreateSerializationContext(UActorComponent* Parent);
+
+	// Returns an ovrAudioContext for the given audio device, or nullptr if one does not exist.
+	static ovrAudioContext GetContextForAudioDevice(const FAudioDevice* InAudioDevice);
+
+	// Creates a new ovrAudioContext for a given audio device.
+	// the InAudioDevice ptr is no longer used for anything besides looking up contexts after this call is completed.
+	static ovrAudioContext CreateContextForAudioDevice(FAudioDevice* InAudioDevice);
+	static void DestroyContextForAudioDevice(const FAudioDevice* InAudioDevice);
+
+private:
+	// FIXME: can we do something better than global static variables?
+	static ovrAudioContext SerializationContext;
+	static FCriticalSection SerializationContextLock;
+
+	static UActorComponent* SerializationParent;
+
+	static TMap<const FAudioDevice*, ovrAudioContext> ContextMap;
+	static FCriticalSection ContextCritSection;
+};
\ No newline at end of file
diff --git a/Engine/Plugins/Runtime/Oculus/OculusAudio/Source/OculusAudio/Private/OculusAudioDllManager.cpp b/Engine/Plugins/Runtime/Oculus/OculusAudio/Source/OculusAudio/Private/OculusAudioDllManager.cpp
new file mode 100644
--- /dev/null
+++ b/Engine/Plugins/Runtime/Oculus/OculusAudio/Source/OculusAudio/Private/OculusAudioDllManager.cpp
@@ -0,0 +1,219 @@
+// Copyright Epic Games, Inc. All Rights Reserved.
+#include "OculusAudioDllManager.h"
+#include "Misc/Paths.h"
+#include "OculusAudioSettings.h"
+#include "Stats/Stats.h"
+
+
+// forward decleration should match OAP_Globals
+extern "C" ovrResult ovrAudio_GetPluginContext(ovrAudioContext* context, unsigned clientType);
+extern "C" ovrResult OSP_FMOD_SetUnitScale(float unitScale);
+
+#define OVRA_CLIENT_TYPE_FMOD          5
+#define OVRA_CLIENT_TYPE_WWISE_UNKNOWN 12
+
+FOculusAudioLibraryManager& FOculusAudioLibraryManager::Get()
+{
+	static FOculusAudioLibraryManager* sInstance;
+	if (sInstance == nullptr)
+	{
+		sInstance = new FOculusAudioLibraryManager();
+	}
+
+	check(sInstance != nullptr);
+	return *sInstance;
+}
+
+FOculusAudioLibraryManager::FOculusAudioLibraryManager()
+	: OculusAudioDllHandle(nullptr), NumInstances(0),
+	bInitialized(false), CachedPluginContext(nullptr)
+{
+}
+
+FOculusAudioLibraryManager::~FOculusAudioLibraryManager()
+{
+	FTSTicker::GetCoreTicker().RemoveTicker(TickDelegateHandle);
+}
+
+void FOculusAudioLibraryManager::Initialize()
+{
+	if (NumInstances == 0)
+	{
+		if (!LoadDll())
+		{
+			UE_LOG(LogAudio, Error, TEXT("Oculus Audio: Failed to load OVR Audio dll"));
+			check(false);
+			return;
+		}
+	}
+
+	NumInstances++;
+	if (!bInitialized)
+	{
+		// Check the version number
+		int32 MajorVersionNumber = -1;
+		int32 MinorVersionNumber = -1;
+		int32 PatchNumber = -1;
+
+		const char* OvrVersionString = OVRA_CALL(ovrAudio_GetVersion)(&MajorVersionNumber, &MinorVersionNumber, &PatchNumber);
+		if (MajorVersionNumber != OVR_AUDIO_MAJOR_VERSION || MinorVersionNumber != OVR_AUDIO_MINOR_VERSION)
+		{
+			UE_LOG(LogAudio, Warning, TEXT("Oculus Audio: Using mismatched OVR Audio SDK Version! %d.%d vs. %d.%d"), OVR_AUDIO_MAJOR_VERSION, OVR_AUDIO_MINOR_VERSION, MajorVersionNumber, MinorVersionNumber);
+			return;
+		}
+		bInitialized = true;
+	}
+}
+
+void FOculusAudioLibraryManager::Shutdown()
+{
+	if (NumInstances == 0)
+	{
+		// This means we failed to load the OVR Audio module during initialization and there's nothing to shutdown.
+		return;
+	}
+
+	NumInstances--;
+
+	if (NumInstances == 0)
+	{
+		// Shutdown OVR audio
+		ReleaseDll();
+		bInitialized = false;
+	}
+}
+
+
+bool FOculusAudioLibraryManager::LoadDll()
+{
+	if (OculusAudioDllHandle == nullptr)
+	{
+		const TCHAR* WWISE_DLL_NAME = TEXT("OculusSpatializerWwise");
+		const TCHAR* FMOD_DLL_NAME = TEXT("OculusSpatializerFMOD");
+		const TCHAR* UE_DLL_NAME = (sizeof(void*) == 4) ? TEXT("ovraudio32") : TEXT("ovraudio64");
+
+#if PLATFORM_WINDOWS 
+
+#if WITH_EDITOR
+		const FString WwisePath = FPaths::ProjectDir() / FString::Printf(TEXT("Binaries/Win64/"));
+		const FString FMODPath = FPaths::ProjectPluginsDir() / FString::Printf(TEXT("FMODStudio/Binaries/Win64/"));
+		const FString UEPath = FPaths::EngineDir() / FString::Printf(TEXT("Binaries/ThirdParty/Oculus/Audio/Win64/"));
+#else
+		const FString WwisePath = FPaths::ProjectDir() / FString::Printf(TEXT("Binaries/Win64/"));
+		const FString FMODPath = FPaths::ProjectPluginsDir() / FString::Printf(TEXT("FMODStudio/Binaries/Win64/")); //TODO verify this
+		const FString UEPath = FPaths::ProjectDir() / FString::Printf(TEXT("../Engine/Binaries/ThirdParty/Oculus/Audio/Win64/"));
+#endif
+
+		FString Path;
+		const TCHAR* DLL_NAME = nullptr;
+		if (FPaths::FileExists(WwisePath + WWISE_DLL_NAME + ".dll"))
+		{
+			ClientType = OVRA_CLIENT_TYPE_WWISE_UNKNOWN;
+			Path = WwisePath;
+			DLL_NAME = WWISE_DLL_NAME;
+			UE_LOG(LogAudio, Display, TEXT("Oculus Audio: OculusSpatializerWwise.dll found, using the Wwise version of the Oculus Audio UE integration"));
+		}
+		else if (FPaths::FileExists(FMODPath + FMOD_DLL_NAME + ".dll"))
+		{
+			ClientType = OVRA_CLIENT_TYPE_FMOD;
+			Path = FMODPath;
+			DLL_NAME = FMOD_DLL_NAME;
+			UE_LOG(LogAudio, Display, TEXT("Oculus Audio: OculusSpatializerFMOD.dll found, using the FMOD version of the Oculus Audio UE integration"));
+		}
+		else
+		{
+			ClientType = -1;
+			Path = UEPath;
+			DLL_NAME = UE_DLL_NAME;
+			UE_LOG(LogAudio, Display, TEXT("Oculus Audio: Middleware plugins not found, assuming native UE AudioMixer"));
+		}
+
+
+		UE_LOG(LogAudio, Display, TEXT("Oculus Audio: Attempting to load Oculus Spatializer DLL: %s (from %s)"), *Path, DLL_NAME);
+
+		FPlatformProcess::PushDllDirectory(*Path);
+		OculusAudioDllHandle = FPlatformProcess::GetDllHandle(*(Path + DLL_NAME));
+		FPlatformProcess::PopDllDirectory(*Path);
+
+#elif PLATFORM_ANDROID
+		FString Path = TEXT("lib");
+		OculusAudioDllHandle = FPlatformProcess::GetDllHandle(*(Path + WWISE_DLL_NAME + ".so"));
+		if (OculusAudioDllHandle != nullptr)
+		{
+			UE_LOG(LogAudio, Display, TEXT("Oculus Audio: %s found, using the Wwise version of the Oculus Audio UE integration"), WWISE_DLL_NAME);
+			return true;
+		}
+		OculusAudioDllHandle = FPlatformProcess::GetDllHandle(*(Path + FMOD_DLL_NAME + ".so"));
+		if (OculusAudioDllHandle != nullptr)
+		{
+			UE_LOG(LogAudio, Display, TEXT("Oculus Audio: %s found, using the FMOD version of the Oculus Audio UE integration"), FMOD_DLL_NAME);
+			return true;
+		}
+		OculusAudioDllHandle = FPlatformProcess::GetDllHandle(*(Path + UE_DLL_NAME + ".so"));
+		if (OculusAudioDllHandle != nullptr)
+		{
+			UE_LOG(LogAudio, Display, TEXT("Oculus Audio: Middleware plugins not found, %s found, assuming native UE AudioMixer"), UE_DLL_NAME);
+			return true;
+		}
+		else
+		{
+			UE_LOG(LogAudio, Error, TEXT("Oculus Audio: Unable to load Oculus Audio UE integratiton"), UE_DLL_NAME);
+		}
+#endif
+
+		return (OculusAudioDllHandle != nullptr);
+	}
+	return true;
+}
+
+void FOculusAudioLibraryManager::ReleaseDll()
+{
+#if PLATFORM_WINDOWS
+	if (NumInstances == 0 && OculusAudioDllHandle)
+	{
+		FPlatformProcess::FreeDllHandle(OculusAudioDllHandle);
+		OculusAudioDllHandle = nullptr;
+	}
+#endif
+}
+
+bool FOculusAudioLibraryManager::UpdatePluginContext(float DeltaTime)
+{
+	QUICK_SCOPE_CYCLE_COUNTER(STAT_FOculusAudioLibraryManager_UpdatePluginContext);
+
+	ovrAudioContext Context = GetPluginContext();
+	ovrResult Result = OVRA_CALL(ovrAudio_UpdateRoomModel)(Context, 1.0f);
+	check(Result == ovrSuccess || Result == ovrError_AudioUninitialized);
+
+	UOculusAudioSettings* settings = GetMutableDefault<UOculusAudioSettings>();
+	Result = OVRA_CALL(ovrAudio_SetPropagationQuality)(Context, settings->PropagationQuality);
+	check(Result == ovrSuccess);
+
+	return true;
+}
+
+ovrAudioContext FOculusAudioLibraryManager::GetPluginContext()
+{
+	if (CachedPluginContext == nullptr)
+	{
+		auto GetPluginContext = OVRA_CALL(ovrAudio_GetPluginContext);
+		if (GetPluginContext != nullptr)
+		{
+			ovrResult Result = GetPluginContext(&CachedPluginContext, ClientType);
+			check(Result == ovrSuccess);
+
+			// FMOD plugin needs to know about the unit scale
+			auto SetFMODUnitScale = OVRA_CALL(OSP_FMOD_SetUnitScale);
+			if (SetFMODUnitScale != nullptr)
+			{
+				Result = SetFMODUnitScale(0.01f);
+				check(Result == ovrSuccess);
+			}
+
+			// Tick the scene from here since there is no listener
+			auto TickDelegate = FTickerDelegate::CreateRaw(this, &FOculusAudioLibraryManager::UpdatePluginContext);
+			TickDelegateHandle = FTSTicker::GetCoreTicker().AddTicker(TickDelegate);
+		}
+	}
+	return CachedPluginContext;
+}
\ No newline at end of file
diff --git a/Engine/Plugins/Runtime/Oculus/OculusAudio/Source/OculusAudio/Private/OculusAudioDllManager.h b/Engine/Plugins/Runtime/Oculus/OculusAudio/Source/OculusAudio/Private/OculusAudioDllManager.h
new file mode 100644
--- /dev/null
+++ b/Engine/Plugins/Runtime/Oculus/OculusAudio/Source/OculusAudio/Private/OculusAudioDllManager.h
@@ -0,0 +1,77 @@
+// Copyright Epic Games, Inc. All Rights Reserved.
+#pragma once
+#include "OVR_Audio.h"
+#include "CoreMinimal.h"
+#include "Audio.h"
+#include "Containers/Ticker.h"
+
+//---------------------------------------------------
+// Oculus Audio DLL handling
+//---------------------------------------------------
+
+static const TCHAR* GetOculusErrorString(ovrResult Result)
+{
+	switch (Result)
+	{
+		default:
+		case ovrError_AudioUnknown:							return TEXT("Unknown Error");
+		case ovrError_AudioInvalidParam:					return TEXT("Invalid Param");
+		case ovrError_AudioBadSampleRate:					return TEXT("Bad Samplerate");
+		case ovrError_AudioMissingDLL:						return TEXT("Missing DLL");
+		case ovrError_AudioBadAlignment:					return TEXT("Pointers did not meet 16 byte alignment requirements");
+		case ovrError_AudioUninitialized:					return TEXT("Function called before initialization");
+		case ovrError_AudioHRTFInitFailure:					return TEXT("HRTF Profider initialization failed");
+		case ovrError_AudioBadVersion:						return TEXT("Bad audio version");
+		case ovrError_AudioSymbolNotFound:					return TEXT("DLL symbol not found");
+		case ovrError_SharedReverbDisabled:					return TEXT("Shared reverb disabled");
+		case ovrError_AudioNoAvailableAmbisonicInstance:	return TEXT("No available Ambisonic");\
+	}
+}
+
+#define OVRA_CALL(Function) \
+	[]() { typedef decltype(&Function) FunctionType; static FunctionType fp = (FunctionType)(FPlatformProcess::GetDllExport(FOculusAudioLibraryManager::Get().DllHandle(), TEXT(#Function))); return fp; }()
+
+
+#define OVR_AUDIO_CHECK_RETURN_VALUE(Result, Context, FailureReturnValue)								\
+	if (Result != ovrSuccess)																			\
+	{																									\
+		const TCHAR* ErrString = GetOculusErrorString(Result);											\
+		UE_LOG(LogAudio, Error, TEXT("Oculus Audio SDK Error - %s: %s"), TEXT(Context), ErrString);		\
+		return FailureReturnValue;																		\
+	}
+
+#define OVR_AUDIO_CHECK(Result, Context) OVR_AUDIO_CHECK_RETURN_VALUE(Result, Context, (void)0)
+
+/************************************************************************/
+/* FOculusAudioLibraryManager										   */
+/* This handles loading and unloading the Oculus Audio DLL at runtime.  */
+/************************************************************************/
+class FOculusAudioLibraryManager
+{
+public:
+	static FOculusAudioLibraryManager& Get();
+
+	void Initialize();
+	void Shutdown();
+
+	bool IsInitialized() { return bInitialized; }
+	ovrAudioContext GetPluginContext();
+
+	void* DllHandle() { return OculusAudioDllHandle; }
+
+private:
+	FOculusAudioLibraryManager();
+	virtual ~FOculusAudioLibraryManager();
+	bool UpdatePluginContext(float DeltaTime);
+
+	void* OculusAudioDllHandle;
+
+	bool LoadDll();
+	void ReleaseDll();
+
+	uint32 NumInstances;
+	bool bInitialized;
+	ovrAudioContext CachedPluginContext;
+	FTSTicker::FDelegateHandle TickDelegateHandle;
+	uint32 ClientType;
+};
\ No newline at end of file
diff --git a/Engine/Plugins/Runtime/Oculus/OculusAudio/Source/OculusAudio/Private/OculusAudioGeometryBSP.cpp b/Engine/Plugins/Runtime/Oculus/OculusAudio/Source/OculusAudio/Private/OculusAudioGeometryBSP.cpp
new file mode 100644
--- /dev/null
+++ b/Engine/Plugins/Runtime/Oculus/OculusAudio/Source/OculusAudio/Private/OculusAudioGeometryBSP.cpp
@@ -0,0 +1,179 @@
+// Copyright Epic Games, Inc. All Rights Reserved.
+#include "OculusAudioGeometryBSP.h"
+#include "OculusAudioContextManager.h"
+#include "OculusAudioMixer.h"
+
+#include "EngineUtils.h"
+#include "Model.h"
+
+#define OVRA_AUDIO_GEOMETRY_BSP_LATEST_VERSION 1
+
+
+static const FGuid UOculusAudioGeometryBSPGUID(0xa03e414e, 0x4ba74409, 0xbd6b2568, 0x2ed8f1f2); // random GUID, guaranteed to be random.
+static const FCustomVersionRegistration UOculusAudioGeometryBSPGUIDRegistration(UOculusAudioGeometryBSPGUID, OVRA_AUDIO_GEOMETRY_BSP_LATEST_VERSION, TEXT("OculusAudioGeometryBSPVersion"));
+
+
+void UOculusAudioGeometryBSP::Serialize(FArchive & Ar)
+{
+	// Tell the archive we are using a custom version.
+	Ar.UsingCustomVersion(UOculusAudioGeometryBSPGUID);
+
+	Super::Serialize(Ar);
+
+	struct Delta {
+		static size_t Read(void* userData, void* bytes, size_t byteCount) {
+			FArchive* Archive = static_cast<FArchive*>(userData);
+			check(Archive->IsLoading());
+
+			// Check that we are not reading past end of archive..
+			int64 CurrentPosition = Archive->Tell();
+			int64 TotalSize = Archive->TotalSize();
+			if ((CurrentPosition + static_cast<int64>(byteCount)) > TotalSize)
+			{
+				return 0;
+			}
+
+			Archive->Serialize(bytes, byteCount);
+			return Archive->GetError() ? 0 : byteCount;
+		}
+		static size_t Write(void* userData, const void* bytes, size_t byteCount) {
+			FArchive* Archive = static_cast<FArchive*>(userData);
+			check(Archive->IsSaving());
+			Archive->Serialize(const_cast<void*>(bytes), byteCount);
+			return Archive->GetError() ? 0 : byteCount;
+		}
+		static int64_t Seek(void* userData, int64_t seekOffset) {
+			FArchive* Archive = static_cast<FArchive*>(userData);
+			int64 Start = Archive->Tell();
+			Archive->Seek(seekOffset);
+			return 0;
+		}
+	};
+
+	ovrAudioSerializer Serializer;
+	Serializer.read = Delta::Read;
+	Serializer.write = Delta::Write;
+	Serializer.seek = Delta::Seek;
+	Serializer.userData = &Ar;
+
+	ovrAudioContext Context = FOculusAudioContextManager::GetOrCreateSerializationContext(this);
+
+	UWorld* World = GetWorld();
+	if (Ar.IsSaving() && World != nullptr && !World->IsGameWorld())
+	{
+		CreateGeometryFromBSP(Context, World, &ovrGeometryBSP);
+
+		if (ovrGeometryBSP == nullptr)
+		{
+			// nothing to serialize
+			return;
+		}
+
+		ovrResult Result = OVRA_CALL(ovrAudio_AudioGeometryWriteMeshData)(ovrGeometryBSP, &Serializer);
+		check(Result == ovrSuccess);
+
+		Result = OVRA_CALL(ovrAudio_DestroyAudioGeometry)(ovrGeometryBSP);
+		check(Result == ovrSuccess);
+
+		ovrGeometryBSP = nullptr;
+	}
+	else if (Ar.IsLoading())
+	{
+		ovrResult Result = OVRA_CALL(ovrAudio_CreateAudioGeometry)(Context, &ovrGeometryBSP);
+		check(Result == ovrSuccess);
+
+		int64 PreviousPosition = Ar.Tell();
+		Result = OVRA_CALL(ovrAudio_AudioGeometryReadMeshData)(ovrGeometryBSP, &Serializer);
+		if (Result != ovrSuccess)
+		{
+			Result = OVRA_CALL(ovrAudio_DestroyAudioGeometry)(ovrGeometryBSP);
+			check(Result == ovrSuccess);
+
+			ovrGeometryBSP = nullptr;
+			Ar.Seek(PreviousPosition);
+		}
+	}
+}
+
+void UOculusAudioGeometryBSP::CreateGeometryFromBSP(ovrAudioContext Context, UWorld* World, ovrAudioGeometry* Geometry)
+{
+	TArray<FVector> Vertices;
+	TArray<int32> Indices;
+	for (auto& WorldVertex : World->GetModel()->Points)
+	{
+		Vertices.Add(OculusAudioSpatializationAudioMixer::ToOVRVector((FVector)WorldVertex));
+	}
+
+	if (Vertices.Num() == 0)
+	{
+		UE_LOG(LogAudio, Warning, TEXT("No BSP geometry found in level!"));
+		return;
+	}
+
+	int FaceCount = 0;
+	for (auto& WorldNode : World->GetModel()->Nodes)
+	{
+		if (WorldNode.NumVertices <= 2)
+		{
+			continue;
+		}
+
+		int32 Index0 = World->GetModel()->Verts[WorldNode.iVertPool + 0].pVertex;
+		int32 Index1 = World->GetModel()->Verts[WorldNode.iVertPool + 1].pVertex;
+		int32 Index2;
+
+		for (auto v = 2; v < WorldNode.NumVertices; ++v)
+		{
+			Index2 = World->GetModel()->Verts[WorldNode.iVertPool + v].pVertex;
+			Indices.Add(Index0);
+			Indices.Add(Index2);
+			Indices.Add(Index1);
+			Index1 = Index2;
+
+			++FaceCount;
+		}
+	}
+
+	ovrAudioMesh ovrMesh = { };
+
+	ovrAudioMeshVertices ovrVertices = { 0 };
+	ovrVertices.vertices = Vertices.GetData();
+	ovrVertices.vertexCount = Vertices.Num();
+	check(ovrVertices.vertexCount != 0);
+	ovrVertices.vertexType = ovrAudioScalarType_Float32;
+	ovrMesh.vertices = ovrVertices;
+
+	ovrAudioMeshIndices ovrIndices = { 0 };
+	ovrIndices.indices = Indices.GetData();
+	ovrIndices.indexCount = Indices.Num();
+	check(ovrIndices.indexCount != 0);
+	ovrIndices.indexType = ovrAudioScalarType_UInt32;
+	ovrMesh.indices = ovrIndices;
+
+	ovrAudioMaterial ovrMaterial = nullptr;
+	ovrResult Result = OVRA_CALL(ovrAudio_CreateAudioMaterial)(Context, &ovrMaterial);
+	check(Result == ovrSuccess);
+
+	ConstructMaterial(ovrMaterial);
+
+	ovrAudioMeshGroup MeshGroup = { 0 };
+	MeshGroup.faceCount = FaceCount;
+	MeshGroup.faceType = ovrAudioFaceType_Triangles;
+	MeshGroup.material = ovrMaterial;
+	ovrMesh.groups = &MeshGroup;
+	ovrMesh.groupCount = 1;
+
+	Result = OVRA_CALL(ovrAudio_CreateAudioGeometry)(Context, Geometry);
+	if (Result != ovrSuccess) {
+		UE_LOG(LogAudio, Warning, TEXT("Failed to create audio propagation geometry for BSP!"));
+		return;
+	}
+
+	Result = OVRA_CALL(ovrAudio_AudioGeometryUploadMesh)(*Geometry, &ovrMesh);
+	if (Result != ovrSuccess) {
+		UE_LOG(LogAudio, Warning, TEXT("Failed adding BSP geometry to the audio propagation sub-system!"));
+		return;
+	}
+
+	UE_LOG(LogAudio, Warning, TEXT("BSP complete!"));
+}
diff --git a/Engine/Plugins/Runtime/Oculus/OculusAudio/Source/OculusAudio/Private/OculusAudioGeometryBSP.h b/Engine/Plugins/Runtime/Oculus/OculusAudio/Source/OculusAudio/Private/OculusAudioGeometryBSP.h
new file mode 100644
--- /dev/null
+++ b/Engine/Plugins/Runtime/Oculus/OculusAudio/Source/OculusAudio/Private/OculusAudioGeometryBSP.h
@@ -0,0 +1,24 @@
+// Copyright Epic Games, Inc. All Rights Reserved.
+
+#pragma once
+
+#include "Components/SceneComponent.h"
+#include "OculusAudioMaterialComponent.h"
+
+#include "OculusAudioGeometryBSP.generated.h"
+
+
+class UWorld;
+
+
+UCLASS(ClassGroup = (Audio), HideCategories = (Activation, Collision, Cooking), meta = (BlueprintSpawnableComponent))
+class UOculusAudioGeometryBSP : public UOculusAudioMaterialComponent
+{
+	GENERATED_BODY()
+
+	void Serialize(FArchive & Ar) override;
+
+	ovrAudioGeometry ovrGeometryBSP = nullptr;
+
+	void CreateGeometryFromBSP(ovrAudioContext Context, UWorld* World, ovrAudioGeometry* Geometry);
+};
\ No newline at end of file
diff --git a/Engine/Plugins/Runtime/Oculus/OculusAudio/Source/OculusAudio/Private/OculusAudioGeometryComponent.cpp b/Engine/Plugins/Runtime/Oculus/OculusAudio/Source/OculusAudio/Private/OculusAudioGeometryComponent.cpp
new file mode 100644
--- /dev/null
+++ b/Engine/Plugins/Runtime/Oculus/OculusAudio/Source/OculusAudio/Private/OculusAudioGeometryComponent.cpp
@@ -0,0 +1,530 @@
+// Copyright Epic Games, Inc. All Rights Reserved.
+#include "OculusAudioGeometryComponent.h"
+#include "AudioDevice.h"
+#include "Components/StaticMeshComponent.h"
+#include "Engine/StaticMesh.h"
+#include "Engine/StaticMeshActor.h"
+#include "Engine/World.h"
+#include "IOculusAudioPlugin.h"
+#include "Materials/MaterialInstanceDynamic.h"
+#include "OculusAudio.h"
+#include "OculusAudioContextManager.h"
+#include "OculusAudioMaterialComponent.h"
+#include "Serialization/CustomVersion.h"
+#include "StaticMeshResources.h"
+
+
+#define COMPILE_AS_DEBUG_MODE 0
+
+#if COMPILE_AS_DEBUG_MODE
+#  pragma optimize("", off) // worse perf, better debugging
+#endif
+
+
+#define OVRA_AUDIO_GEOMETRY_COMPONENT_LATEST_VERSION 1
+
+static const FGuid UOculusAudioGeometryGUID( 0x7aa25488, 0x40d74391, 0xad87d335, 0x347cfae1 ); // random GUID, guaranteed to be random.
+static const FCustomVersionRegistration UOculusAudioGeometryGUIDRegistration( UOculusAudioGeometryGUID, OVRA_AUDIO_GEOMETRY_COMPONENT_LATEST_VERSION, TEXT("OculusAudioGeometryVersion") );
+
+UOculusAudioGeometryComponent::UOculusAudioGeometryComponent()
+	: ovrGeometry(nullptr)
+	, CachedContext(nullptr)
+{
+	bAutoActivate = true;
+	PrimaryComponentTick.bCanEverTick = true;
+	bWantsInitializeComponent = true;
+}
+
+UOculusAudioGeometryComponent::~UOculusAudioGeometryComponent()
+{
+}
+
+// ------------------------------------------------------------------------------------------------------------------------------------------------------------
+// Note: a child (attached) static mesh will have it's geometry merged into the parent. Therefore, its ovrGeometry data member will be nullptr below and this is why
+// we need the Uploaded member.
+void UOculusAudioGeometryComponent::TickComponent(float DeltaTime, enum ELevelTick TickType, FActorComponentTickFunction *ThisTickFunction)
+{
+	Super::TickComponent(DeltaTime, TickType, ThisTickFunction);
+
+	if (ovrGeometry == nullptr)
+	{
+		AStaticMeshActor* Actor = Cast<AStaticMeshActor>(GetOwner());
+		check(Actor != nullptr);
+
+		UStaticMesh* Mesh = Actor->GetStaticMeshComponent()->GetStaticMesh();
+		check(Mesh != nullptr);
+		check(Mesh->bAllowCPUAccess);
+
+		UploadGeometry();
+	}
+
+	check(ovrGeometry != nullptr);
+
+	// 3D transform
+	FTransform Transform = GetOwner()->GetTransform();
+	if (!Transform.Equals(PreviousTransform))
+	{
+		// UE:		x:forward, y:right, z:up
+		// Oculus:  x:right,   y:up,	z:backward
+		FVector Right = Transform.GetScaledAxis(EAxis::Y);
+		FVector Up = Transform.GetScaledAxis(EAxis::Z);
+		FVector Backward = -Transform.GetScaledAxis(EAxis::X);
+		FVector Position = Transform.GetLocation();
+		
+		float TransformMatrix[16] =
+		{
+			(float)Right.Y,	(float)Right.Z,	(float)-Right.X,	0.0f,
+			(float)Up.Y,	   (float)Up.Z,	   (float)-Up.X,	   0.0f,
+			(float)Backward.Y, (float)Backward.Z, (float)-Backward.X, 0.0f,
+			(float)Position.Y, (float)Position.Z, (float)-Position.X, 0.0f,
+		};
+		
+		ovrResult Result = OVRA_CALL(ovrAudio_AudioGeometrySetTransform)(ovrGeometry, TransformMatrix);
+		if (Result != ovrSuccess) {
+			UE_LOG(LogAudio, Warning, TEXT("Failed at setting new audio propagation mesh transform!"));
+			return;
+		}
+
+		PreviousTransform = Transform;
+	}
+}
+
+// ------------------------------------------------------------------------------------------------------------------------------------------------------------
+bool UOculusAudioGeometryComponent::UploadGeometry()
+{
+	ovrAudioContext Context = GetContext();
+	if (Context == nullptr)
+		return false;
+
+	check(ovrGeometry == nullptr);
+
+	// static mesh actor
+	AStaticMeshActor* SMActor = Cast<AStaticMeshActor>(GetOwner());
+	check(SMActor != nullptr);
+
+	// actual mesh
+	UStaticMesh* Mesh = SMActor->GetStaticMeshComponent()->GetStaticMesh();
+	if (Mesh == nullptr)
+		return false; //PAS TEMP
+	check(Mesh != nullptr);
+
+	// LB: silently filter out invisible geometry (might be used for something else)
+	if (!Mesh->HasValidRenderData()) {
+		UE_LOG(LogAudio, Warning, TEXT("Found static mesh actor with acoustic geometry component while invisible."));
+		return false;
+	}
+
+	// our internal geometry object into which we'll merge this static mesh actor and all its attached meshes
+	ovrResult Result = OVRA_CALL(ovrAudio_CreateAudioGeometry)(Context, &ovrGeometry);
+	if (Result != ovrSuccess) {
+		UE_LOG(LogAudio, Warning, TEXT("Failed creating acoustic geometry object."));
+		return false;
+	}
+
+	ovrAudioMaterial ovrMaterial = nullptr;
+	// if there is a material component apply that material
+	UOculusAudioMaterialComponent* MaterialComponent = static_cast<UOculusAudioMaterialComponent*>(SMActor->GetComponentByClass(UOculusAudioMaterialComponent::StaticClass()));
+	if (MaterialComponent)
+	{
+		Result = OVRA_CALL(ovrAudio_CreateAudioMaterial)(Context, &ovrMaterial);
+		check(Result == ovrSuccess);
+
+		MaterialComponent->ConstructMaterial(ovrMaterial);
+	}
+
+#if 0 // TODO: map acoustic materials to visual/phyiscal materials
+	// The goal here is to build a hash table with unique graphic materials used by the acoustic geometry and automatically mapping them with
+	// our corresponding audio materials. A possible alternative is to do the mapping manually by adding a user data field in the UE material
+	// definition and manually tagging each graphic material with an audio material in the editor. 
+	UStaticMeshComponent* StaticMeshComponent = SMActor->GetStaticMeshComponent();
+	for (int x = 0; x < StaticMeshComponent->GetNumMaterials(); ++x) {
+		//UMaterialInstanceDynamic* Material = UMaterialInstanceDynamic::Create(StaticMeshComponent->GetMaterial(x), nullptr);
+		UMaterialInterface* Material = StaticMeshComponent->GetMaterial(x);
+		//UE_LOG(LogAudio, Warning, TEXT("Found material: %s"), *Material->GetName());
+	}
+#endif
+
+	TArray<FVector> MergedVertices;
+	TArray<uint32> MergedIndices;
+	TArray<ovrAudioMeshGroup> ovrMeshGroups;
+	int32 numVerticesCur = 0;
+
+	// root (parent) mesh always
+	AppendStaticMesh(Mesh, FTransform(), MergedVertices, MergedIndices, ovrMeshGroups,
+		ovrMaterial);
+
+	if (IncludeChildren)
+	{
+		AActor* Actor = GetOwner();
+		AppendChildMeshes(Actor, Actor->GetTransform(), Context,
+			MergedVertices, MergedIndices, ovrMeshGroups,
+			ovrMaterial);
+	}
+
+	ovrAudioMesh ovrMesh = { };
+	ovrAudioMeshVertices ovrVertices = { 0 };
+
+	ovrVertices.vertices = MergedVertices.GetData();
+	ovrVertices.vertexCount = MergedVertices.Num();
+	check(ovrVertices.vertexCount != 0);
+	ovrVertices.vertexType = ovrAudioScalarType_Float32;
+
+	ovrMesh.vertices = ovrVertices;
+
+	ovrAudioMeshIndices ovrIndices = { 0 };
+	ovrIndices.indices = MergedIndices.GetData();
+	ovrIndices.indexCount = MergedIndices.Num();
+	check(ovrIndices.indexCount != 0);
+	ovrIndices.indexType = ovrAudioScalarType_UInt32;
+
+	ovrMesh.indices = ovrIndices;
+	ovrMesh.groups = ovrMeshGroups.GetData();
+	ovrMesh.groupCount = ovrMeshGroups.Num();
+	check(ovrMesh.groupCount != 0);
+
+	Result = OVRA_CALL(ovrAudio_AudioGeometryUploadMesh)(ovrGeometry, &ovrMesh);
+	if (Result != ovrSuccess) {
+		UE_LOG(LogAudio, Warning, TEXT("Failed adding geometry to the audio propagation sub-system!"));
+	}
+
+	return (Result == ovrSuccess);
+}
+
+// ------------------------------------------------------------------------------------------------------------------------------------------------------------
+void UOculusAudioGeometryComponent::AppendStaticMesh(UStaticMesh* Mesh,
+													 const FTransform& Transform,
+													 TArray<FVector>& MergedVertices,
+													 TArray<uint32>& MergedIndices,
+													 TArray<ovrAudioMeshGroup>& ovrMeshGroups,
+													 ovrAudioMaterial ovrMaterial)
+{
+	const int32 IndexOffset = MergedVertices.Num();
+	const int32 FirstIndex = MergedIndices.Num();
+
+	// Lowest LOD
+	const FStaticMeshLODResources& Model = Mesh->GetRenderData()->LODResources[0];
+	const FPositionVertexBuffer& VertexBuffer = Model.VertexBuffers.PositionVertexBuffer;
+	const int32 VertexCount = VertexBuffer.GetNumVertices();
+
+	// append this mesh's vertices to the merged array after coordinate system conversion
+	for (int32 Index = 0; Index < VertexCount; Index++) {
+		const FVector Vertex = Transform.TransformPosition((FVector)VertexBuffer.VertexPosition(Index));
+		MergedVertices.Add(OculusAudioSpatializationAudioMixer::ToOVRVector(Vertex));
+	}
+
+	// offset this mesh's vertex indices with the current count within the whole hierarchy
+	// PAS: TODO support 16-bit indices Model.IndexBuffer.Is32Bit();
+	TArray<uint32> MeshIndices;
+	Model.IndexBuffer.GetCopy(MeshIndices);
+	const int32 IndexCount = MeshIndices.Num();
+	for (int32 Index = 0; Index < IndexCount; Index++) {
+		MeshIndices[Index] += IndexOffset;
+	}
+	
+	// append this mesh's indices to the merged array
+	MergedIndices += MeshIndices;
+
+	// loop over each section in lowest LOD
+	for (int32 Index = 0; Index < Model.Sections.Num(); Index++)
+	{
+		ovrAudioMeshGroup MeshGroup = { 0 };
+
+		MeshGroup.faceCount		= Model.Sections[Index].NumTriangles;
+		MeshGroup.faceType		= ovrAudioFaceType_Triangles;
+		MeshGroup.indexOffset	= Model.Sections[Index].FirstIndex + FirstIndex;
+		MeshGroup.material		= ovrMaterial; // TODO: material per section
+
+		ovrMeshGroups.Add(MeshGroup);
+	}
+}
+
+void UOculusAudioGeometryComponent::AppendChildMeshes(AActor* CurrentActor, const FTransform& RootTransform, ovrAudioContext Context,
+	TArray<FVector>& MergedVertices, TArray<uint32>& MergedIndices, TArray<ovrAudioMeshGroup>& ovrMeshGroups, 
+	ovrAudioMaterial ovrMaterial)
+{
+	TArray<AActor*> AttachedActors;
+	CurrentActor->GetAttachedActors(AttachedActors);
+
+	for (AActor* Actor : AttachedActors)
+	{
+		AStaticMeshActor* SMActor = Cast<AStaticMeshActor>(Actor);
+		UStaticMesh* Mesh = SMActor->GetStaticMeshComponent()->GetStaticMesh();
+		check(Mesh != nullptr);
+
+		UOculusAudioGeometryComponent* GeometryComponent = static_cast<UOculusAudioGeometryComponent*>(SMActor->GetComponentByClass(UOculusAudioGeometryComponent::StaticClass()));
+		if (GeometryComponent != nullptr)
+		{
+			// this child has it's own geometry component
+			continue;
+		}
+
+		// silently filter out invisible attached geometry (might be used for something else)
+		// or an attached (child) mesh with no geometry component (acoustically inactive)
+		if (!Mesh->HasValidRenderData()) {
+			continue;
+		}
+		
+		ovrAudioMaterial ovrChildMaterial = ovrMaterial;
+
+		// if we have a child geometry component with a valid material preset, it overrides the root mesh's material preset (inherited from parent by default).
+		UOculusAudioMaterialComponent* ChildMaterial = static_cast<UOculusAudioMaterialComponent*>(SMActor->GetComponentByClass(UOculusAudioMaterialComponent::StaticClass()));
+		if (ChildMaterial != nullptr)
+		{
+			ovrResult Result = OVRA_CALL(ovrAudio_CreateAudioMaterial)(Context, &ovrChildMaterial);
+			check(Result == ovrSuccess);
+
+			ChildMaterial->ConstructMaterial(ovrChildMaterial);
+		}
+
+		// serialize the transform relative to the root
+		FTransform Transform = Actor->GetTransform().GetRelativeTransform(RootTransform);
+
+		// merge this attached static mesh into the root propagation mesh
+		AppendStaticMesh(Mesh,
+			Transform,
+			MergedVertices,
+			MergedIndices,
+			ovrMeshGroups,
+			ovrChildMaterial);
+
+		// traverse children and append those too
+		AppendChildMeshes(Actor, RootTransform, Context, MergedVertices, MergedIndices, ovrMeshGroups,
+			ovrChildMaterial);
+	}
+}
+
+ovrAudioContext UOculusAudioGeometryComponent::GetContext(UWorld* World)
+{
+	if (CachedContext == nullptr)
+	{
+		ovrAudioContext PluginContext = FOculusAudioLibraryManager::Get().GetPluginContext();
+		if (PluginContext != nullptr)
+		{
+			CachedContext = PluginContext;
+		}
+		else
+		{
+			FOculusAudioPlugin* Plugin = &FModuleManager::GetModuleChecked<FOculusAudioPlugin>("OculusAudio");
+			check(Plugin != nullptr);
+
+			// Note: getting the context from the spatializer, should we get it from the reverb instead?
+			FString OculusSpatializerPluginName = Plugin->GetSpatializationPluginFactory()->GetDisplayName();
+			FString CurrentSpatializerPluginName = AudioPluginUtilities::GetDesiredPluginName(EAudioPlugin::SPATIALIZATION);
+			if (CurrentSpatializerPluginName.Equals(OculusSpatializerPluginName)) // we have a match!
+			{
+				if (World == nullptr)
+				{
+					World = GetOwner()->GetWorld();
+				}
+				if (World == nullptr)
+				{
+					// couldn't get a World
+					return nullptr;
+				}
+
+				FAudioDevice* AudioDevice = World->GetAudioDevice().GetAudioDevice();
+				if (AudioDevice == nullptr)
+				{
+					// This happens when cooking for native UE AudioMixer integration
+					CachedContext = FOculusAudioContextManager::GetOrCreateSerializationContext(this);
+				}
+				else
+				{
+					CachedContext = FOculusAudioContextManager::GetContextForAudioDevice(AudioDevice);
+
+					if (!CachedContext)
+					{
+						CachedContext = FOculusAudioContextManager::CreateContextForAudioDevice(AudioDevice);
+					}
+				}
+
+				OculusAudioSpatializationAudioMixer* Spatializer = static_cast<OculusAudioSpatializationAudioMixer*>(AudioDevice->GetSpatializationPluginInterface().Get());
+				if (Spatializer == nullptr || Spatializer->ClassID != OculusAudioSpatializationAudioMixer::MIXER_CLASS_ID)
+				{
+					UE_LOG(LogAudio, Warning, TEXT("Invalid Spatialization Plugin specified, make sure the Spatialization Plugin is set to OculusAudio and AudioMixer is enabled!"));
+					return nullptr;
+				}
+			}
+		}
+	}
+
+	return CachedContext;
+}
+
+void UOculusAudioGeometryComponent::Serialize(FArchive & Ar)
+{
+	// Tell the archive we are using a custom version.
+	Ar.UsingCustomVersion( UOculusAudioGeometryGUID );
+
+	Super::Serialize(Ar);
+
+	struct Delta {
+		static size_t Read(void* userData, void* bytes, size_t byteCount) {
+			FArchive* Archive = static_cast<FArchive*>(userData);
+			check(Archive->IsLoading());
+			int64 CurrentPosition = Archive->Tell();
+			int64 TotalSize = Archive->TotalSize();
+			if ((CurrentPosition + static_cast<int64>(byteCount)) > TotalSize)
+			{
+				// Note: After copy/paste/undo TotalSize returns -1 and Serialize fails
+				return 0;
+			}
+
+			Archive->Serialize(bytes, byteCount);
+			return Archive->GetError() ? 0 : byteCount;
+		}
+		static size_t Write(void* userData, const void* bytes, size_t byteCount) {
+			FArchive* Archive = static_cast<FArchive*>(userData);
+			check(Archive->IsSaving());
+			Archive->Serialize(const_cast<void*>(bytes), byteCount);
+			return Archive->GetError() ? 0 : byteCount;
+		}
+		static int64_t Seek(void* userData, int64_t seekOffset) {
+			FArchive* Archive = static_cast<FArchive*>(userData);
+			int64 Start = Archive->Tell();
+			Archive->Seek(seekOffset);
+			return 0;
+		}
+	};
+
+	ovrAudioSerializer Serializer;
+	Serializer.read = Delta::Read;
+	Serializer.write = Delta::Write;
+	Serializer.seek = Delta::Seek;
+	Serializer.userData = &Ar;
+
+	if (Ar.IsSaving())
+	{
+		AStaticMeshActor* Actor = Cast<AStaticMeshActor>(GetOwner());
+		if (Actor == nullptr)
+			return; // this happens when adding blueprint to actor with component already attached
+
+		UStaticMeshComponent* MeshComponent = Actor->GetStaticMeshComponent();
+		if (MeshComponent == nullptr) // this happens with actor copy/paste/delete
+			return;
+
+		UStaticMesh* Mesh = MeshComponent->GetStaticMesh();
+
+		// LB 8/1/18: when duplicating static mesh actors in the editor, we get here once with a NULL Mesh ptr, not sure why (could be that the original and the
+		// dupe instances share the (single) actual mesh data and a serialization callback happens before the dupe has updated the mesh ptr?)
+		//check(Mesh != nullptr);
+		if (Mesh == nullptr) 
+			return;
+
+		// if we have anything to save, it means the content changed for this component in the editor so we need to do the full export from scratch
+		// but 1st cleanup any existing (stale) data:
+		if (ovrGeometry != nullptr) {
+			ovrResult Result = OVRA_CALL(ovrAudio_DestroyAudioGeometry)(ovrGeometry);
+			check(Result == ovrSuccess);
+			ovrGeometry = nullptr;
+		}
+
+		// upload the new geometry and material preset	
+		UploadGeometry();
+
+#if 0 // PAS TEMP
+		{
+			char path[256] = { 0 };
+			static int counter = 0;
+			sprintf(path, "G:/TEMP/Propagation/MeshDump/%s-%d.obj", TCHAR_TO_ANSI(*Mesh->GetName()), counter++);
+			ovrResult Result = OVRA_CALL(ovrAudio_AudioGeometryWriteMeshFileObj)(ovrGeometry, path);
+			check(Result == ovrSuccess);
+		}
+#endif
+#if 0 // PAS TEMP
+		{
+			char path[256] = { 0 };
+			static int counter = 0;
+			sprintf(path, "G:/Dev/Oculus/depot/Software/AudioSDK/Main/Tests/PropagTestData/%s-%d.geo", TCHAR_TO_ANSI(*Mesh->GetName()), counter++);
+			ovrResult Result = OVRA_CALL(ovrAudio_AudioGeometryWriteMeshFile)(ovrGeometry, path);
+			check(Result == ovrSuccess);
+		}
+#endif
+
+		// now save cooked data
+		if (ovrGeometry != nullptr) {
+			ovrResult Result = OVRA_CALL(ovrAudio_AudioGeometryWriteMeshData)(ovrGeometry, &Serializer);
+			check(Result == ovrSuccess);
+
+			Result = OVRA_CALL(ovrAudio_DestroyAudioGeometry)(ovrGeometry);
+			check(Result == ovrSuccess);
+			ovrGeometry = nullptr;
+		}
+	}
+	else if (Ar.IsLoading())
+	{
+		// Read the file data so unreal doesn't complain that we didn't read all the data if the version differs.
+		ovrAudioContext Context = FOculusAudioContextManager::GetOrCreateSerializationContext(this);
+		ovrResult Result = OVRA_CALL(ovrAudio_CreateAudioGeometry)(Context, &ovrGeometry);
+		check(Result == ovrSuccess);
+
+		int64 PreviousPosition = Ar.Tell();
+		Result = OVRA_CALL(ovrAudio_AudioGeometryReadMeshData)(ovrGeometry, &Serializer);
+		if (Result != ovrSuccess)
+		{
+			Result = OVRA_CALL(ovrAudio_DestroyAudioGeometry)(ovrGeometry);
+			check(Result == ovrSuccess);
+
+			ovrGeometry = nullptr;
+			Ar.Seek(PreviousPosition);
+		}
+		
+		// If the version number is out of date, upload the mesh from the scene.
+		int32 customVersion = Ar.CustomVer( UOculusAudioGeometryGUID );
+		if (ovrGeometry != nullptr && customVersion != OVRA_AUDIO_GEOMETRY_COMPONENT_LATEST_VERSION)
+		{
+			// if we have anything to save, it means the content changed for this component in the editor so we need to do the full export from scratch
+			// but 1st cleanup any existing (stale) data:
+			if (ovrGeometry != nullptr) {
+				Result = OVRA_CALL(ovrAudio_DestroyAudioGeometry)(ovrGeometry);
+				check(Result == ovrSuccess);
+				ovrGeometry = nullptr;
+			}
+
+			// upload the new geometry and material preset	
+			UploadGeometry();
+		}
+	}
+}
+
+void UOculusAudioGeometryComponent::PostLoad()
+{
+	Super::PostLoad();
+
+#if WITH_EDITOR
+	if (ovrGeometry != nullptr)
+	{
+		AActor* Actor = GetOwner();
+		if (Actor != nullptr)
+		{
+			UWorld* World = Actor->GetWorld();
+			if (World != nullptr)
+			{
+				if (World->WorldType != EWorldType::Game &&
+					World->WorldType != EWorldType::PIE)
+				{
+					ovrResult Result = OVRA_CALL(ovrAudio_DestroyAudioGeometry)(ovrGeometry);
+					check(Result == ovrSuccess);
+					ovrGeometry = nullptr;
+				}
+			}
+		}
+	}
+#endif
+}
+
+void UOculusAudioGeometryComponent::BeginDestroy()
+{
+	Super::BeginDestroy();
+
+	if (ovrGeometry != nullptr && CachedContext != nullptr)
+	{
+		ovrResult Result = OVRA_CALL(ovrAudio_DestroyAudioGeometry)(ovrGeometry);
+		check(Result == ovrSuccess);
+		ovrGeometry = nullptr;
+	}
+}
+
+#if COMPILE_AS_DEBUG_MODE
+#  pragma optimize("", on) //PAS - TEMP - DELETE THIS - better debugging
+#endif
diff --git a/Engine/Plugins/Runtime/Oculus/OculusAudio/Source/OculusAudio/Private/OculusAudioGeometryComponent.h b/Engine/Plugins/Runtime/Oculus/OculusAudio/Source/OculusAudio/Private/OculusAudioGeometryComponent.h
new file mode 100644
--- /dev/null
+++ b/Engine/Plugins/Runtime/Oculus/OculusAudio/Source/OculusAudio/Private/OculusAudioGeometryComponent.h
@@ -0,0 +1,52 @@
+// Copyright Epic Games, Inc. All Rights Reserved.
+#pragma once
+
+#include "Components/ActorComponent.h"
+#include "OVR_Audio.h"
+#include "OculusAudioGeometryComponent.generated.h"
+
+/*
+ * OculusAudio geometry components are used to customize an a static mesh actor's acoustic properties.
+ */
+UCLASS(ClassGroup = (Audio), HideCategories = (Activation, Collision, Cooking), meta = (BlueprintSpawnableComponent))
+class UOculusAudioGeometryComponent : public UActorComponent
+{
+	GENERATED_BODY()
+
+public:
+	UOculusAudioGeometryComponent();
+	~UOculusAudioGeometryComponent();
+	bool UploadGeometry();
+	bool IncludesChildren() const	{ return IncludeChildren; }
+
+private:
+	void TickComponent(float DeltaTime, enum ELevelTick TickType, FActorComponentTickFunction *ThisTickFunction) override;
+	void Serialize(FArchive & Ar) override;
+	void PostLoad() override;
+	void BeginDestroy() override;
+
+	ovrAudioContext GetContext(UWorld* World = nullptr);
+	void AppendStaticMesh(UStaticMesh* Mesh,
+						  const FTransform& Transform,
+						  TArray<FVector>& MergedVertices,
+						  TArray<uint32>& MergedIndices,
+						  TArray<ovrAudioMeshGroup>& ovrMeshGroups,
+						  ovrAudioMaterial ovrMaterial);
+
+	void AppendChildMeshes(AActor* CurrentActor,
+						   const FTransform& RootTransform,
+						   ovrAudioContext Context,
+						   TArray<FVector>& MergedVertices,
+						   TArray<uint32>& MergedIndices,
+						   TArray<ovrAudioMeshGroup>& ovrMeshGroups,
+						   ovrAudioMaterial ovrMaterial);
+
+	// Mesh hierarchy optimization for both content editing and runtime performance
+	// if IncludeChildren is true, children (attached) meshes will be merged
+	UPROPERTY(EditAnywhere, Category = Settings)
+	bool IncludeChildren = true;
+
+	ovrAudioGeometry ovrGeometry;
+	ovrAudioContext CachedContext;
+	FTransform PreviousTransform;
+};
\ No newline at end of file
diff --git a/Engine/Plugins/Runtime/Oculus/OculusAudio/Source/OculusAudio/Private/OculusAudioGeometryLandscape.cpp b/Engine/Plugins/Runtime/Oculus/OculusAudio/Source/OculusAudio/Private/OculusAudioGeometryLandscape.cpp
new file mode 100644
--- /dev/null
+++ b/Engine/Plugins/Runtime/Oculus/OculusAudio/Source/OculusAudio/Private/OculusAudioGeometryLandscape.cpp
@@ -0,0 +1,191 @@
+// Copyright Epic Games, Inc. All Rights Reserved.
+#include "OculusAudioGeometryLandscape.h"
+#include "OculusAudioContextManager.h"
+#include "OculusAudioMixer.h"
+
+#include "EngineUtils.h"
+#include "Landscape.h"
+#include "LandscapeComponent.h"
+#include "LandscapeInfo.h"
+#include "LandscapeDataAccess.h"
+
+#define OVRA_AUDIO_GEOMETRY_LANDSCAPE_LATEST_VERSION 1
+
+static const FGuid UOculusAudioGeometryLandscapeGUID(0xa7387e03, 0x7a95402c, 0x9292fe5f, 0x170d05a9); // random GUID, guaranteed to be random.
+static const FCustomVersionRegistration UOculusAudioGeometryLandscapeGUIDRegistration(UOculusAudioGeometryLandscapeGUID, OVRA_AUDIO_GEOMETRY_LANDSCAPE_LATEST_VERSION, TEXT("OculusAudioGeometryLandscapeVersion"));
+
+
+void UOculusAudioGeometryLandscape::Serialize(FArchive & Ar)
+{
+	// Tell the archive we are using a custom version.
+	Ar.UsingCustomVersion(UOculusAudioGeometryLandscapeGUID);
+
+	Super::Serialize(Ar);
+
+	struct Delta {
+		static size_t Read(void* userData, void* bytes, size_t byteCount) {
+			FArchive* Archive = static_cast<FArchive*>(userData);
+			check(Archive->IsLoading());
+
+			// Check that we are not reading past end of archive..
+			int64 CurrentPosition = Archive->Tell();
+			int64 TotalSize = Archive->TotalSize();
+			if ((CurrentPosition + static_cast<int64>(byteCount)) > TotalSize)
+			{
+				return 0;
+			}
+
+			Archive->Serialize(bytes, byteCount);
+			return Archive->GetError() ? 0 : byteCount;
+		}
+		static size_t Write(void* userData, const void* bytes, size_t byteCount) {
+			FArchive* Archive = static_cast<FArchive*>(userData);
+			check(Archive->IsSaving());
+			Archive->Serialize(const_cast<void*>(bytes), byteCount);
+			return Archive->GetError() ? 0 : byteCount;
+		}
+		static int64_t Seek(void* userData, int64_t seekOffset) {
+			FArchive* Archive = static_cast<FArchive*>(userData);
+			int64 Start = Archive->Tell();
+			Archive->Seek(seekOffset);
+			return 0;
+		}
+	};
+
+	ovrAudioSerializer Serializer;
+	Serializer.read = Delta::Read;
+	Serializer.write = Delta::Write;
+	Serializer.seek = Delta::Seek;
+	Serializer.userData = &Ar;
+
+	ovrAudioContext Context = FOculusAudioContextManager::GetOrCreateSerializationContext(this);
+
+#if WITH_EDITOR
+	UWorld* World = GetWorld();
+	if (Ar.IsSaving() && World != nullptr && !World->IsGameWorld())
+	{
+		CreateGeometryFromLandscape(Context, World, &ovrGeometryLandscape);
+
+		if (ovrGeometryLandscape != nullptr)
+		{
+			ovrResult Result = OVRA_CALL(ovrAudio_AudioGeometryWriteMeshData)(ovrGeometryLandscape, &Serializer);
+			check(Result == ovrSuccess);
+
+			Result = OVRA_CALL(ovrAudio_DestroyAudioGeometry)(ovrGeometryLandscape);
+			check(Result == ovrSuccess);
+
+			ovrGeometryLandscape = nullptr;
+		}
+	}
+#endif
+	
+	if (Ar.IsLoading())
+	{
+		ovrResult Result = OVRA_CALL(ovrAudio_CreateAudioGeometry)(Context, &ovrGeometryLandscape);
+		check(Result == ovrSuccess);
+
+		int64 PreviousPosition = Ar.Tell();
+		Result = OVRA_CALL(ovrAudio_AudioGeometryReadMeshData)(ovrGeometryLandscape, &Serializer);
+		if (Result != ovrSuccess)
+		{
+			Result = OVRA_CALL(ovrAudio_DestroyAudioGeometry)(ovrGeometryLandscape);
+			check(Result == ovrSuccess);
+
+			ovrGeometryLandscape = nullptr;
+			Ar.Seek(PreviousPosition);
+		}
+	}
+}
+
+#if WITH_EDITOR
+void UOculusAudioGeometryLandscape::CreateGeometryFromLandscape(ovrAudioContext Context, UWorld* World, ovrAudioGeometry* Geometry)
+{
+	for (TActorIterator<ALandscape> Itr(World); Itr; ++Itr)
+	{
+		TArray<FVector> Vertices;
+		TArray<uint32> Indices;
+		int32 FaceCount = 0;
+
+		auto LandscapeCollisionComponents = Itr->CollisionComponents;
+
+		ULandscapeInfo* LandscapeInfo = Itr->CreateLandscapeInfo();
+		if (LandscapeInfo == nullptr)
+		{
+			UE_LOG(LogAudio, Warning, TEXT("Unabled to create audio propagation geometry for landscape, landscape info not found!"));
+			return;
+		}
+
+		for (auto It = LandscapeInfo->XYtoComponentMap.CreateIterator(); It; ++It)
+		{
+			ULandscapeComponent* Component = It.Value();
+			FLandscapeComponentDataInterface DataInterface(Component);
+
+			const int Offset = Vertices.Num();
+			const int Length = Component->ComponentSizeQuads + 1;
+			for (auto y = 0; y < Length; ++y)
+			{
+				for (auto x = 0; x < Length; ++x)
+				{
+					Vertices.Add(OculusAudioSpatializationAudioMixer::ToOVRVector(DataInterface.GetWorldVertex(x, y)));
+					
+					// there are only (N-1)^2 quads for N^2 verts
+					if (x == (Length - 1) || y == (Length - 1))
+						continue;
+
+					auto calcIndex = [](int x, int y, int length) { return (y * length) + x; };
+					Indices.Add(Offset + calcIndex(x, y, Length));
+					Indices.Add(Offset + calcIndex(x, y + 1, Length));
+					Indices.Add(Offset + calcIndex(x + 1, y + 1, Length));
+					Indices.Add(Offset + calcIndex(x + 1, y, Length));
+
+					++FaceCount;
+				}
+			}
+		}
+
+		ovrAudioMesh ovrMesh = { };
+
+		ovrAudioMeshVertices ovrVertices = { 0 };
+		ovrVertices.vertices = Vertices.GetData();
+		ovrVertices.vertexCount = Vertices.Num();
+		check(ovrVertices.vertexCount != 0);
+		ovrVertices.vertexType = ovrAudioScalarType_Float32;
+		ovrMesh.vertices = ovrVertices;
+
+		ovrAudioMeshIndices ovrIndices = { 0 };
+		ovrIndices.indices = Indices.GetData();
+		ovrIndices.indexCount = Indices.Num();
+		check(ovrIndices.indexCount != 0);
+		ovrIndices.indexType = ovrAudioScalarType_UInt32;
+		ovrMesh.indices = ovrIndices;
+		
+		ovrAudioMaterial ovrMaterial = nullptr;
+		ovrResult Result = OVRA_CALL(ovrAudio_CreateAudioMaterial)(Context, &ovrMaterial);
+		check(Result == ovrSuccess);
+
+		ConstructMaterial(ovrMaterial);
+
+		ovrAudioMeshGroup MeshGroup = { 0 };
+		MeshGroup.faceCount = FaceCount;
+		MeshGroup.faceType = ovrAudioFaceType_Quads;
+		MeshGroup.material = ovrMaterial;
+
+		ovrMesh.groups = &MeshGroup;
+		ovrMesh.groupCount = 1;
+
+		Result = OVRA_CALL(ovrAudio_CreateAudioGeometry)(Context, Geometry);
+		if (Result != ovrSuccess) {
+			UE_LOG(LogAudio, Warning, TEXT("Failed to create audio propagation geometry for landscape!"));
+			return;
+		}
+
+		Result = OVRA_CALL(ovrAudio_AudioGeometryUploadMesh)(*Geometry, &ovrMesh);
+		if (Result != ovrSuccess) {
+			UE_LOG(LogAudio, Warning, TEXT("Failed adding landscape geometry to the audio propagation sub-system!"));
+			return;
+		}
+	}
+
+	UE_LOG(LogAudio, Warning, TEXT("Landscape complete!"));
+}
+#endif
diff --git a/Engine/Plugins/Runtime/Oculus/OculusAudio/Source/OculusAudio/Private/OculusAudioGeometryLandscape.h b/Engine/Plugins/Runtime/Oculus/OculusAudio/Source/OculusAudio/Private/OculusAudioGeometryLandscape.h
new file mode 100644
--- /dev/null
+++ b/Engine/Plugins/Runtime/Oculus/OculusAudio/Source/OculusAudio/Private/OculusAudioGeometryLandscape.h
@@ -0,0 +1,23 @@
+// Copyright Epic Games, Inc. All Rights Reserved.
+#pragma once
+
+#include "Components/SceneComponent.h"
+#include "OculusAudioMaterialComponent.h"
+
+#include "OculusAudioGeometryLandscape.generated.h"
+
+
+class UWorld;
+
+
+UCLASS(ClassGroup = (Audio), HideCategories = (Activation, Collision, Cooking), meta = (BlueprintSpawnableComponent))
+class UOculusAudioGeometryLandscape : public UOculusAudioMaterialComponent
+{
+	GENERATED_BODY()
+
+	void Serialize(FArchive & Ar) override;
+
+	ovrAudioGeometry ovrGeometryLandscape = nullptr;
+
+	void CreateGeometryFromLandscape(ovrAudioContext Context, UWorld* World, ovrAudioGeometry* Geometry);
+};
\ No newline at end of file
diff --git a/Engine/Plugins/Runtime/Oculus/OculusAudio/Source/OculusAudio/Private/OculusAudioMaterialComponent.cpp b/Engine/Plugins/Runtime/Oculus/OculusAudio/Source/OculusAudio/Private/OculusAudioMaterialComponent.cpp
new file mode 100644
--- /dev/null
+++ b/Engine/Plugins/Runtime/Oculus/OculusAudio/Source/OculusAudio/Private/OculusAudioMaterialComponent.cpp
@@ -0,0 +1,357 @@
+// Copyright Epic Games, Inc. All Rights Reserved.
+#include "OculusAudioMaterialComponent.h"
+#include "OculusAudioDllManager.h"
+
+UOculusAudioMaterialComponent::UOculusAudioMaterialComponent()
+	: MaterialPreset(EOculusAudioMaterial::NOMATERIAL) // default
+{
+	ResetAcousticMaterialPreset();
+}
+
+void UOculusAudioMaterialComponent::ConstructMaterial(ovrAudioMaterial ovrMaterial)
+{
+	ovrResult Result = ovrSuccess;
+	Result = OVRA_CALL(ovrAudio_AudioMaterialSetFrequency)(ovrMaterial, ovrAudioMaterialProperty_Absorption, 63.5f,   Absorption63Hz);
+	check(Result == ovrSuccess);
+	Result = OVRA_CALL(ovrAudio_AudioMaterialSetFrequency)(ovrMaterial, ovrAudioMaterialProperty_Absorption, 125.0f,  Absorption125Hz);
+	check(Result == ovrSuccess);
+	Result = OVRA_CALL(ovrAudio_AudioMaterialSetFrequency)(ovrMaterial, ovrAudioMaterialProperty_Absorption, 250.0f,  Absorption250Hz);
+	check(Result == ovrSuccess);
+	Result = OVRA_CALL(ovrAudio_AudioMaterialSetFrequency)(ovrMaterial, ovrAudioMaterialProperty_Absorption, 500.0f,  Absorption500Hz);
+	check(Result == ovrSuccess);
+	Result = OVRA_CALL(ovrAudio_AudioMaterialSetFrequency)(ovrMaterial, ovrAudioMaterialProperty_Absorption, 1000.0f, Absorption1000Hz);
+	check(Result == ovrSuccess);
+	Result = OVRA_CALL(ovrAudio_AudioMaterialSetFrequency)(ovrMaterial, ovrAudioMaterialProperty_Absorption, 2000.0f, Absorption2000Hz);
+	check(Result == ovrSuccess);
+	Result = OVRA_CALL(ovrAudio_AudioMaterialSetFrequency)(ovrMaterial, ovrAudioMaterialProperty_Absorption, 4000.0f, Absorption4000Hz);
+	check(Result == ovrSuccess);
+	Result = OVRA_CALL(ovrAudio_AudioMaterialSetFrequency)(ovrMaterial, ovrAudioMaterialProperty_Absorption, 8000.0f, Absorption8000Hz);
+	check(Result == ovrSuccess);
+
+	Result = OVRA_CALL(ovrAudio_AudioMaterialSetFrequency)(ovrMaterial, ovrAudioMaterialProperty_Transmission, 63.5f,   Transmission63Hz);
+	check(Result == ovrSuccess);
+	Result = OVRA_CALL(ovrAudio_AudioMaterialSetFrequency)(ovrMaterial, ovrAudioMaterialProperty_Transmission, 125.0f,  Transmission125Hz);
+	check(Result == ovrSuccess);
+	Result = OVRA_CALL(ovrAudio_AudioMaterialSetFrequency)(ovrMaterial, ovrAudioMaterialProperty_Transmission, 250.0f,  Transmission250Hz);
+	check(Result == ovrSuccess);
+	Result = OVRA_CALL(ovrAudio_AudioMaterialSetFrequency)(ovrMaterial, ovrAudioMaterialProperty_Transmission, 500.0f,  Transmission500Hz);
+	check(Result == ovrSuccess);
+	Result = OVRA_CALL(ovrAudio_AudioMaterialSetFrequency)(ovrMaterial, ovrAudioMaterialProperty_Transmission, 1000.0f, Transmission1000Hz);
+	check(Result == ovrSuccess);
+	Result = OVRA_CALL(ovrAudio_AudioMaterialSetFrequency)(ovrMaterial, ovrAudioMaterialProperty_Transmission, 2000.0f, Transmission2000Hz);
+	check(Result == ovrSuccess);
+	Result = OVRA_CALL(ovrAudio_AudioMaterialSetFrequency)(ovrMaterial, ovrAudioMaterialProperty_Transmission, 4000.0f, Transmission4000Hz);
+	check(Result == ovrSuccess);
+	Result = OVRA_CALL(ovrAudio_AudioMaterialSetFrequency)(ovrMaterial, ovrAudioMaterialProperty_Transmission, 8000.0f, Transmission8000Hz);
+	check(Result == ovrSuccess);
+
+	Result = OVRA_CALL(ovrAudio_AudioMaterialSetFrequency)(ovrMaterial, ovrAudioMaterialProperty_Scattering, 63.5f,   Scattering63Hz);
+	check(Result == ovrSuccess);
+	Result = OVRA_CALL(ovrAudio_AudioMaterialSetFrequency)(ovrMaterial, ovrAudioMaterialProperty_Scattering, 125.0f,  Scattering125Hz);
+	check(Result == ovrSuccess);
+	Result = OVRA_CALL(ovrAudio_AudioMaterialSetFrequency)(ovrMaterial, ovrAudioMaterialProperty_Scattering, 250.0f,  Scattering250Hz);
+	check(Result == ovrSuccess);
+	Result = OVRA_CALL(ovrAudio_AudioMaterialSetFrequency)(ovrMaterial, ovrAudioMaterialProperty_Scattering, 500.0f,  Scattering500Hz);
+	check(Result == ovrSuccess);
+	Result = OVRA_CALL(ovrAudio_AudioMaterialSetFrequency)(ovrMaterial, ovrAudioMaterialProperty_Scattering, 1000.0f, Scattering1000Hz);
+	check(Result == ovrSuccess);
+	Result = OVRA_CALL(ovrAudio_AudioMaterialSetFrequency)(ovrMaterial, ovrAudioMaterialProperty_Scattering, 2000.0f, Scattering2000Hz);
+	check(Result == ovrSuccess);
+	Result = OVRA_CALL(ovrAudio_AudioMaterialSetFrequency)(ovrMaterial, ovrAudioMaterialProperty_Scattering, 4000.0f, Scattering4000Hz);
+	check(Result == ovrSuccess);
+	Result = OVRA_CALL(ovrAudio_AudioMaterialSetFrequency)(ovrMaterial, ovrAudioMaterialProperty_Scattering, 8000.0f, Scattering8000Hz);
+	check(Result == ovrSuccess);
+}
+
+void UOculusAudioMaterialComponent::ResetAcousticMaterialPreset()
+{
+	Absorption63Hz = .0f;
+	Absorption125Hz = .0f;
+	Absorption250Hz = .0f;
+	Absorption500Hz = .0f;
+	Absorption1000Hz = .0f;
+	Absorption2000Hz = .0f;
+	Absorption4000Hz = .0f;
+	Absorption8000Hz = .0f;
+
+	Transmission63Hz = .0f;
+	Transmission125Hz = .0f;
+	Transmission250Hz = .0f;
+	Transmission500Hz = .0f;
+	Transmission1000Hz = .0f;
+	Transmission2000Hz = .0f;
+	Transmission4000Hz = .0f;
+	Transmission8000Hz = .0f;
+
+	Scattering63Hz = .0f;
+	Scattering125Hz = .0f;
+	Scattering250Hz = .0f;
+	Scattering500Hz = .0f;
+	Scattering1000Hz = .0f;
+	Scattering2000Hz = .0f;
+	Scattering4000Hz = .0f;
+	Scattering8000Hz = .0f;
+}
+
+#if WITH_EDITOR
+
+	void UOculusAudioMaterialComponent::PostEditChangeProperty(struct FPropertyChangedEvent& PropertyChangedEvent)
+	{
+		FName PropertyName = (PropertyChangedEvent.Property != NULL) ? PropertyChangedEvent.Property->GetFName() : NAME_None;
+
+		if (PropertyName == GET_MEMBER_NAME_CHECKED(UOculusAudioMaterialComponent, MaterialPreset))
+		{
+			if (IsValidMaterialPreset()) // a valid existing material until we allow custom ones
+			{
+				switch (MaterialPreset)
+				{
+				case EOculusAudioMaterial::ACOUSTICTILE:
+					AssignCoefficients(
+						{ .50f, .50f, .70f, .60f, .70f, .70f, .50f, .50f },
+						{ .05f, .05f, .04f, .03f, .02f, .005f, .002f, .001f },
+						{ .10f, .10f, .15f, .20f, .20f, .25f, .30f, .35f });
+					break;
+				case EOculusAudioMaterial::BRICK:
+					AssignCoefficients(
+						{ .02f, .02f, .02f, .03f, .04f, .05f, .07f, .07f },
+						{ .025f, .025f, .019f, .01f, .0045f, .0018f, .00089f, .00089f },
+						{ .10f, .20f, .25f, .30f, .35f, .40f, .45f, .50f });
+					break;
+				case EOculusAudioMaterial::BRICKPAINTED:
+					AssignCoefficients(
+						{ .01f, .01f, .01f, .02f, .02f, .02f, .03f, .03f },
+						{ .025f, .025f, .019f, .01f, .0045f, .0018f, .00089f, .00089f },
+						{ .10f, .15f, .15f, .20f, .20f, .20f, .25f, .30f });
+					break;
+				case EOculusAudioMaterial::CARPET:
+					AssignCoefficients(
+						{ .01f, .01f, .05f, .10f, .20f, .45f, .65f, .70f },
+						{ .004f, .004f, .0079f, .0056f, .0016f, .0014f, .0005f, .0003f },
+						{ .10f, .10f, .10f, .15f, .20f, .30f, .45f, .50f });
+					break;
+				case EOculusAudioMaterial::CARPETHEAVY:
+					AssignCoefficients(
+						{ .02f, .02f, .06f, .14f, .37f, .48f, .63f, .63f },
+						{ .004f, .004f, .0079f, .0056f, .0016f, .0014f, .0005f, .0003f },
+						{ .10f, .10f, .10f, .15f, .20f, .30f, .45f, .50f });
+					break;
+				case EOculusAudioMaterial::CARPETHEAVYPADDED:
+					AssignCoefficients(
+						{ .08f, .08f, .24f, .57f, .69f, .71f, .73f, .73f },
+						{ .004f, .004f, .0079f, .0056f, .0016f, .0014f, .0005f, .0003f },
+						{ .10f, .10f, .10f, .15f, .20f, .30f, .45f, .50f });
+					break;
+				case EOculusAudioMaterial::CERAMICTILE:
+					AssignCoefficients(
+						{ .01f, .01f, .01f, .01f, .01f, .02f, .02f, .02f },
+						{ .004f, .004f, .0079f, .0056f, .0016f, .0014f, .0005f, .0003f },
+						{ .10f, .10f, .12f, .14f, .16f, .18f, .20f, .22f });
+					break;
+				case EOculusAudioMaterial::CONCRETE:
+					AssignCoefficients(
+						{ .01f, .01f, .01f, .02f, .02f, .02f, .02f, .02f },
+						{ .004f, .004f, .0079f, .0056f, .0016f, .0014f, .0005f, .0003f },
+						{ .10f, .10f, .11f, .12f, .13f, .14f, .15f, .20f });
+					break;
+				case EOculusAudioMaterial::CONCRETEROUGH:
+					AssignCoefficients(
+						{ .01f, .01f, .02f, .04f, .06f, .08f, .10f, .10f },
+						{ .004f, .004f, .0079f, .0056f, .0016f, .0014f, .0005f, .0003f },
+						{ .10f, .10f, .12f, .15f, .20f, .25f, .30f, .35f });
+					break;
+				case EOculusAudioMaterial::CONCRETEBLOCK:
+					AssignCoefficients(
+						{ .36f, .36f, .44f, .31f, .29f, .39f, .21f, .21f },
+						{ .02f, .02f, .01f, .0063f, .0035f, .00011f, .00063f, .0005f },
+						{ .10f, .10f, .12f, .15f, .20f, .30f, .40f, .45f });
+					break;
+				case EOculusAudioMaterial::CONCRETEBLOCKPAINTED:
+					AssignCoefficients(
+						{ .10f, .10f, .05f, .06f, .07f, .09f, .08f, .08f },
+						{ .02f, .02f, .01f, .0063f, .0035f, .00011f, .00063f, .00063f },
+						{ .10f, .10f, .11f, .13f, .15f, .16f, .20f, .25f });
+					break;
+				case EOculusAudioMaterial::CURTAIN:
+					AssignCoefficients(
+						{ .05f, .07f, .31f, .49f, .75f, .70f, .60f, .60f },
+						{ .50f, .42f, .39f, .21f, .14f, .079f, .045f, .04f },
+						{ .10f, .10f, .15f, .2f, .3f, .4f, .5f, .5f });
+					break;
+				case EOculusAudioMaterial::FOLIAGE:
+					AssignCoefficients(
+						{ .03f, .03f, .06f, .11f, .17f, .27f, .31f, .31f },
+						{ .95f, .9f, .9f, .9f, .8f, .5f, .3f, .25f },
+						{ .20f, .20f, .3f, .4f, .5f, .7f, .8f, .8f });
+					break;
+				case EOculusAudioMaterial::GLASS:
+					AssignCoefficients(
+						{ .35f, .35f, .25f, .18f, .12f, .07f, .05f, .05f },
+						{ .200f, .125f, .089f, .05f, .028f, .022f, .079f, .05f },
+						{ .05f, .05f, .05f, .05f, .05f, .05f, .05f, .05f });
+					break;
+				case EOculusAudioMaterial::GLASSHEAVY:
+					AssignCoefficients(
+						{ .18f, .18f, .06f, .04f, .03f, .02f, .02f, .02f },
+						{ .06f, .056f, .039f, .028f, .02f, .032f, .014f, .010f },
+						{ .05f, .05f, .05f, .05f, .05f, .05f, .05f, .05f });
+					break;
+				case EOculusAudioMaterial::GRASS:
+					AssignCoefficients(
+						{ .11f, .11f, .26f, .60f, .69f, .92f, .99f, .99f },
+						{ .0f, .0f, .0f, .0f, .0f, .0f, .0f, .0f },
+						{ .20f, .25f, .30f, .40f, .50f, .60f, .70f, .80f });
+					break;
+				case EOculusAudioMaterial::GRAVEL:
+					AssignCoefficients(
+						{ .25f, .25f, .60f, .65f, .70f, .75f, .80f, .80f },
+						{ .0f, .0f, .0f, .0f, .0f, .0f, .0f, .0f },
+						{ .20f, .25f, .30f, .40f, .50f, .60f, .70f, .80f });
+					break;
+				case EOculusAudioMaterial::GYPSUMBOARD:
+					AssignCoefficients(
+						{ .29f, .29f, .10f, .05f, .04f, .07f, .09f, .09f },
+						{ .05f, .035f, .0125f, .0056f, .0025f, .0013f, .0032f, .002f },
+						{ .10f, .10f, .11f, .12f, .13f, .14f, .15f, .20f });
+					break;
+				case EOculusAudioMaterial::PLASTERONBRICK:
+					AssignCoefficients(
+						{ .01f, .01f, .02f, .02f, .03f, .04f, .05f, .05f },
+						{ .025f, .025f, .019f, .01f, .0045f, .0018f, .00089f, .0006f },
+						{ .10f, .10f, .11f, .12f, .13f, .14f, .15f, .20f });
+					break;
+				case EOculusAudioMaterial::PLASTERONCONCRETEBLOCK:
+					AssignCoefficients(
+						{ .12f, .12f, .09f, .07f, .05f, .05f, .04f, .04f },
+						{ .02f, .02f, .01f, .0063f, .0035f, .00011f, .00063f, .0005f },
+						{ .10f, .10f, .11f, .12f, .13f, .14f, .15f, .20f });
+					break;
+				case EOculusAudioMaterial::SOIL:
+					AssignCoefficients(
+						{ .15f, .15f, .25f, .40f, .55f, .60f, .60f, .60f },
+						{ .0f, .0f, .0f, .0f, .0f, .0f, .0f, .0f },
+						{ .10f, .15f, .20f, .25f, .40f, .55f, .70f, .80f });
+					break;
+				case EOculusAudioMaterial::SOUNDPROOF:
+					AssignCoefficients(
+						{ 1.f, 1.f, 1.f, 1.f, 1.f, 1.f, 1.f, 1.f },
+						{ .0f, .0f, .0f, .0f, .0f, .0f, .0f, .0f },
+						{ .0f, .0f, .0f, .0f, .0f, .0f, .0f, .0f });
+					break;
+				case EOculusAudioMaterial::SNOW:
+					AssignCoefficients(
+						{ .45f, .45f, .75f, .90f, .95f, .95f, .95f, .95f },
+						{ .0f, .0f, .0f, .0f, .0f, .0f, .0f, .0f },
+						{ .15f, .20f, .30f, .40f, .50f, .60f, .75f, .80f });
+					break;
+				case EOculusAudioMaterial::STEEL:
+					AssignCoefficients(
+						{ .05f, .05f, .10f, .10f, .10f, .07f, .02f, .02f },
+						{ .3f, .25f, .2f, .17f, .089f, .089f, .0056f, .003f },
+						{ .10f, .10f, .10f, .10f, .10f, .10f, .10f, .10f });
+					break;
+				case EOculusAudioMaterial::WATER:
+					AssignCoefficients(
+						{ .01f, .01f, .01f, .01f, .02f, .02f, .03f, .03f },
+						{ .03f, .03f, .03f, .03f, .02f, .015f, .01f, .01f },
+						{ .10f, .10f, .10f, .10f, .07f, .05f, .05f, .05f });
+					break;
+				case EOculusAudioMaterial::WOODTHIN:
+					AssignCoefficients(
+						{ .42f, .42f, .21f, .10f, .08f, .06f, .06f, .06f },
+						{ .22f, .2f, .125f, .079f, .1f, .089f, .05f, .03f },
+						{ .10f, .10f, .10f, .10f, .10f, .10f, .15f, .15f });
+					break;
+				case EOculusAudioMaterial::WOODTHICK:
+					AssignCoefficients(
+						{ .19f, .19f, .14f, .09f, .06f, .06f, .05f, .05f },
+						{ .04f, .035f, .028f, .028f, .028f, .011f, .0071f, .005f },
+						{ .10f, .10f, .10f, .10f, .10f, .10f, .15f, .15f });
+					break;
+				case EOculusAudioMaterial::WOODFLOOR:
+					AssignCoefficients(
+						{ .15f, .15f, .11f, .10f, .07f, .06f, .07f, .07f },
+						{ .08f, .071f, .025f, .0158f, .0056f, .0035f, .0016f, .001f },
+						{ .10f, .10f, .10f, .10f, .10f, .10f, .15f, .15f });
+					break;
+				case EOculusAudioMaterial::WOODONCONCRETE:
+					AssignCoefficients(
+						{ .04f, .04f, .04f, .07f, .06f, .06f, .07f, .07f },
+						{ .004f, .004f, .0079f, .0056f, .0016f, .0014f, .0005f, .0003f },
+						{ .10f, .10f, .10f, .10f, .10f, .10f, .15f, .15f });
+					break;
+				}
+			} 
+			else
+			{
+				ResetAcousticMaterialPreset();
+			}
+		}
+	}
+
+	bool UOculusAudioMaterialComponent::CanEditChange(const FProperty* InProperty) const
+	{
+		const bool ParentVal = Super::CanEditChange(InProperty);
+		FName PropertyName = InProperty->GetFName();
+
+		if (PropertyName == GET_MEMBER_NAME_CHECKED(UOculusAudioMaterialComponent, Absorption63Hz) ||
+			PropertyName == GET_MEMBER_NAME_CHECKED(UOculusAudioMaterialComponent, Absorption125Hz) ||
+			PropertyName == GET_MEMBER_NAME_CHECKED(UOculusAudioMaterialComponent, Absorption250Hz) ||
+			PropertyName == GET_MEMBER_NAME_CHECKED(UOculusAudioMaterialComponent, Absorption500Hz) ||
+			PropertyName == GET_MEMBER_NAME_CHECKED(UOculusAudioMaterialComponent, Absorption1000Hz) ||
+			PropertyName == GET_MEMBER_NAME_CHECKED(UOculusAudioMaterialComponent, Absorption2000Hz) ||
+			PropertyName == GET_MEMBER_NAME_CHECKED(UOculusAudioMaterialComponent, Absorption4000Hz) ||
+			PropertyName == GET_MEMBER_NAME_CHECKED(UOculusAudioMaterialComponent, Absorption8000Hz) ||
+			PropertyName == GET_MEMBER_NAME_CHECKED(UOculusAudioMaterialComponent, Transmission63Hz) ||
+			PropertyName == GET_MEMBER_NAME_CHECKED(UOculusAudioMaterialComponent, Transmission125Hz) ||
+			PropertyName == GET_MEMBER_NAME_CHECKED(UOculusAudioMaterialComponent, Transmission250Hz) ||
+			PropertyName == GET_MEMBER_NAME_CHECKED(UOculusAudioMaterialComponent, Transmission500Hz) ||
+			PropertyName == GET_MEMBER_NAME_CHECKED(UOculusAudioMaterialComponent, Transmission1000Hz) ||
+			PropertyName == GET_MEMBER_NAME_CHECKED(UOculusAudioMaterialComponent, Transmission2000Hz) ||
+			PropertyName == GET_MEMBER_NAME_CHECKED(UOculusAudioMaterialComponent, Transmission4000Hz) ||
+			PropertyName == GET_MEMBER_NAME_CHECKED(UOculusAudioMaterialComponent, Transmission8000Hz) ||
+			PropertyName == GET_MEMBER_NAME_CHECKED(UOculusAudioMaterialComponent, Scattering63Hz) ||
+			PropertyName == GET_MEMBER_NAME_CHECKED(UOculusAudioMaterialComponent, Scattering125Hz) ||
+			PropertyName == GET_MEMBER_NAME_CHECKED(UOculusAudioMaterialComponent, Scattering250Hz) ||
+			PropertyName == GET_MEMBER_NAME_CHECKED(UOculusAudioMaterialComponent, Scattering500Hz) ||
+			PropertyName == GET_MEMBER_NAME_CHECKED(UOculusAudioMaterialComponent, Scattering1000Hz) ||
+			PropertyName == GET_MEMBER_NAME_CHECKED(UOculusAudioMaterialComponent, Scattering2000Hz) ||
+			PropertyName == GET_MEMBER_NAME_CHECKED(UOculusAudioMaterialComponent, Scattering4000Hz) ||
+			PropertyName == GET_MEMBER_NAME_CHECKED(UOculusAudioMaterialComponent, Scattering8000Hz))
+		{
+			return ParentVal  &&  IsValidMaterialPreset();
+		}
+		else
+		{
+			return ParentVal;
+		}
+	}
+
+	void UOculusAudioMaterialComponent::AssignCoefficients(TArray<float> Absorption, TArray<float> Transmission, TArray<float> Scattering)
+	{
+		Absorption63Hz   = Absorption[0];
+		Absorption125Hz  = Absorption[1];
+		Absorption250Hz  = Absorption[2];
+		Absorption500Hz  = Absorption[3];
+		Absorption1000Hz = Absorption[4];
+		Absorption2000Hz = Absorption[5];
+		Absorption4000Hz = Absorption[6];
+		Absorption8000Hz = Absorption[7];
+		Transmission63Hz   = Transmission[0];
+		Transmission125Hz  = Transmission[1];
+		Transmission250Hz  = Transmission[2];
+		Transmission500Hz  = Transmission[3];
+		Transmission1000Hz = Transmission[4];
+		Transmission2000Hz = Transmission[5];
+		Transmission4000Hz = Transmission[6];
+		Transmission8000Hz = Transmission[7];
+		Scattering63Hz   = Scattering[0];
+		Scattering125Hz  = Scattering[1];
+		Scattering250Hz  = Scattering[2];
+		Scattering500Hz  = Scattering[3];
+		Scattering1000Hz = Scattering[4];
+		Scattering2000Hz = Scattering[5];
+		Scattering4000Hz = Scattering[6];
+		Scattering8000Hz = Scattering[7];
+	}
+
+#endif // WITH_EDITOR
\ No newline at end of file
diff --git a/Engine/Plugins/Runtime/Oculus/OculusAudio/Source/OculusAudio/Private/OculusAudioMaterialComponent.h b/Engine/Plugins/Runtime/Oculus/OculusAudio/Source/OculusAudio/Private/OculusAudioMaterialComponent.h
new file mode 100644
--- /dev/null
+++ b/Engine/Plugins/Runtime/Oculus/OculusAudio/Source/OculusAudio/Private/OculusAudioMaterialComponent.h
@@ -0,0 +1,172 @@
+// Copyright Epic Games, Inc. All Rights Reserved.
+#pragma once
+
+#include "Components/ActorComponent.h"
+#include "OVR_Audio.h"
+#include "OculusAudioMaterialComponent.generated.h"
+
+UENUM()
+enum class EOculusAudioMaterial : uint8
+{
+	ACOUSTICTILE = 0			UMETA(DisplayName = "AcousticTile"),
+	BRICK = 1					UMETA(DisplayName = "Brick"),
+	BRICKPAINTED = 2			UMETA(DisplayName = "BrickPainted"),
+	CARPET = 3					UMETA(DisplayName = "Carpet"),
+	CARPETHEAVY = 4				UMETA(DisplayName = "CarpetHeavy"),
+	CARPETHEAVYPADDED = 5		UMETA(DisplayName = "CarpetHeavyPadded"),
+	CERAMICTILE = 6				UMETA(DisplayName = "CeramicTile"),
+	CONCRETE = 7				UMETA(DisplayName = "Concrete"),
+	CONCRETEROUGH = 8			UMETA(DisplayName = "ConcreteRough"),
+	CONCRETEBLOCK = 9			UMETA(DisplayName = "ConcreteBlock"),
+	CONCRETEBLOCKPAINTED = 10	UMETA(DisplayName = "ConcreteBlockPainted"),
+	CURTAIN = 11				UMETA(DisplayName = "Curtain"),
+	FOLIAGE = 12				UMETA(DisplayName = "Foliage"),
+	GLASS = 13					UMETA(DisplayName = "Glass"),
+	GLASSHEAVY = 14				UMETA(DisplayName = "GlassHeavy"),
+	GRASS = 15					UMETA(DisplayName = "Grass"),
+	GRAVEL = 16					UMETA(DisplayName = "Gravel"),
+	GYPSUMBOARD = 17			UMETA(DisplayName = "GypsumBoard"),
+	PLASTERONBRICK = 18			UMETA(DisplayName = "PlasterOnBrick"),
+	PLASTERONCONCRETEBLOCK = 19 UMETA(DisplayName = "PlasterOnConcreteBlock"),
+	SOIL = 20					UMETA(DisplayName = "Soil"),
+	SOUNDPROOF = 21				UMETA(DisplayName = "SoundProof"),
+	SNOW = 22					UMETA(DisplayName = "Snow"),
+	STEEL = 23					UMETA(DisplayName = "Steel"),
+	WATER = 24					UMETA(DisplayName = "Water"),
+	WOODTHIN = 25				UMETA(DisplayName = "WoodThin"),
+	WOODTHICK = 26				UMETA(DisplayName = "WoodThick"),
+	WOODFLOOR = 27				UMETA(DisplayName = "WoodFloor"),
+	WOODONCONCRETE = 28			UMETA(DisplayName = "WoodOnConcrete"),
+	MATERIAL_MAX				UMETA(Hidden),
+	NOMATERIAL = 255			UMETA(DisplayName = "NoMaterial") // default
+};
+
+/*
+ * OculusAudio material components are used to set the acoustic properties of the geometry.
+ */
+// PAS: TODO check these UCLASS parameters
+UCLASS(ClassGroup = (Audio), HideCategories = (Activation, Collision, Cooking), meta = (BlueprintSpawnableComponent))
+class UOculusAudioMaterialComponent : public UActorComponent
+{
+	GENERATED_BODY()
+public:
+	UOculusAudioMaterialComponent();
+	void ConstructMaterial(ovrAudioMaterial ovrMaterial);
+
+#if WITH_EDITOR
+	virtual void PostEditChangeProperty(struct FPropertyChangedEvent& PropertyChangedEvent) override;
+	virtual bool CanEditChange(const FProperty* InProperty) const override;
+	void AssignCoefficients(TArray<float> Absorption, TArray<float> Transmission, TArray<float> Scattering);
+#endif
+
+private:
+	EOculusAudioMaterial GetMaterialPreset() const { return MaterialPreset; }
+	bool IsValidMaterialPreset() const { return (MaterialPreset < EOculusAudioMaterial::MATERIAL_MAX); }
+	void ResetAcousticMaterialPreset();
+
+	// Choose from a variety of preset physical materials, or choose Custom to specify values manually.
+	UPROPERTY(EditAnywhere, Category = Settings)
+	EOculusAudioMaterial MaterialPreset;
+
+	// ------------------------------------------		Absorption			-----------------------------------------------
+
+	// How much this material absorbs @ 63Hz
+	UPROPERTY(VisibleAnywhere, Category = Settings, meta = (ClampMin = "0.0", ClampMax = "1.0", UIMin = "0.0", UIMax = "1.0"))
+	float Absorption63Hz;
+
+	// How much this material absorbs @ 125Hz
+	UPROPERTY(VisibleAnywhere, Category = Settings, meta = (ClampMin = "0.0", ClampMax = "1.0", UIMin = "0.0", UIMax = "1.0"))
+	float Absorption125Hz;
+
+	// How much this material absorbs @ 250Hz
+	UPROPERTY(VisibleAnywhere, Category = Settings, meta = (ClampMin = "0.0", ClampMax = "1.0", UIMin = "0.0", UIMax = "1.0"))
+	float Absorption250Hz;
+
+	// How much this material absorbs @ 500Hz
+	UPROPERTY(VisibleAnywhere, Category = Settings, meta = (ClampMin = "0.0", ClampMax = "1.0", UIMin = "0.0", UIMax = "1.0"))
+	float Absorption500Hz;
+
+	// How much this material absorbs @ 1000Hz
+	UPROPERTY(VisibleAnywhere, Category = Settings, meta = (ClampMin = "0.0", ClampMax = "1.0", UIMin = "0.0", UIMax = "1.0"))
+	float Absorption1000Hz;
+
+	// How much this material absorbs @ 2000Hz
+	UPROPERTY(VisibleAnywhere, Category = Settings, meta = (ClampMin = "0.0", ClampMax = "1.0", UIMin = "0.0", UIMax = "1.0"))
+	float Absorption2000Hz;
+
+	// How much this material absorbs @ 4000Hz
+	UPROPERTY(VisibleAnywhere, Category = Settings, meta = (ClampMin = "0.0", ClampMax = "1.0", UIMin = "0.0", UIMax = "1.0"))
+	float Absorption4000Hz;
+
+	// How much this material absorbs @ 8000Hz
+	UPROPERTY(VisibleAnywhere, Category = Settings, meta = (ClampMin = "0.0", ClampMax = "1.0", UIMin = "0.0", UIMax = "1.0"))
+	float Absorption8000Hz;
+
+	// ------------------------------------------		Transmission		-----------------------------------------------
+
+	// How much this material transmits @ 63Hz
+	UPROPERTY(VisibleAnywhere, Category = Settings, meta = (ClampMin = "0.0", ClampMax = "1.0", UIMin = "0.0", UIMax = "1.0"))
+	float Transmission63Hz;
+
+	// How much this material transmits @ 125Hz
+	UPROPERTY(VisibleAnywhere, Category = Settings, meta = (ClampMin = "0.0", ClampMax = "1.0", UIMin = "0.0", UIMax = "1.0"))
+	float Transmission125Hz;
+
+	// How much this material transmits @ 250Hz
+	UPROPERTY(VisibleAnywhere, Category = Settings, meta = (ClampMin = "0.0", ClampMax = "1.0", UIMin = "0.0", UIMax = "1.0"))
+	float Transmission250Hz;
+
+	// How much this material transmits @ 500Hz
+	UPROPERTY(VisibleAnywhere, Category = Settings, meta = (ClampMin = "0.0", ClampMax = "1.0", UIMin = "0.0", UIMax = "1.0"))
+	float Transmission500Hz;
+
+	// How much this material transmits @ 1000Hz
+	UPROPERTY(VisibleAnywhere, Category = Settings, meta = (ClampMin = "0.0", ClampMax = "1.0", UIMin = "0.0", UIMax = "1.0"))
+	float Transmission1000Hz;
+
+	// How much this material transmits @ 2000Hz
+	UPROPERTY(VisibleAnywhere, Category = Settings, meta = (ClampMin = "0.0", ClampMax = "1.0", UIMin = "0.0", UIMax = "1.0"))
+	float Transmission2000Hz;
+
+	// How much this material transmits @ 4000Hz
+	UPROPERTY(VisibleAnywhere, Category = Settings, meta = (ClampMin = "0.0", ClampMax = "1.0", UIMin = "0.0", UIMax = "1.0"))
+	float Transmission4000Hz;
+
+	// How much this material transmits @ 8000Hz
+	UPROPERTY(VisibleAnywhere, Category = Settings, meta = (ClampMin = "0.0", ClampMax = "1.0", UIMin = "0.0", UIMax = "1.0"))
+	float Transmission8000Hz;
+
+	// ------------------------------------------		Scattering			-----------------------------------------------
+
+	// How much this material scatters @ 63Hz
+	UPROPERTY(VisibleAnywhere, Category = Settings, meta = (ClampMin = "0.0", ClampMax = "1.0", UIMin = "0.0", UIMax = "1.0"))
+	float Scattering63Hz;
+
+	// How much this material scatters @ 125Hz
+	UPROPERTY(VisibleAnywhere, Category = Settings, meta = (ClampMin = "0.0", ClampMax = "1.0", UIMin = "0.0", UIMax = "1.0"))
+	float Scattering125Hz;
+
+	// How much this material scatters @ 250Hz
+	UPROPERTY(VisibleAnywhere, Category = Settings, meta = (ClampMin = "0.0", ClampMax = "1.0", UIMin = "0.0", UIMax = "1.0"))
+	float Scattering250Hz;
+
+	// How much this material scatters @ 500Hz
+	UPROPERTY(VisibleAnywhere, Category = Settings, meta = (ClampMin = "0.0", ClampMax = "1.0", UIMin = "0.0", UIMax = "1.0"))
+	float Scattering500Hz;
+
+	// How much this material scatters @ 1000Hz
+	UPROPERTY(VisibleAnywhere, Category = Settings, meta = (ClampMin = "0.0", ClampMax = "1.0", UIMin = "0.0", UIMax = "1.0"))
+	float Scattering1000Hz;
+
+	// How much this material scatters @ 2000Hz
+	UPROPERTY(VisibleAnywhere, Category = Settings, meta = (ClampMin = "0.0", ClampMax = "1.0", UIMin = "0.0", UIMax = "1.0"))
+	float Scattering2000Hz;
+
+	// How much this material scatters @ 4000Hz
+	UPROPERTY(VisibleAnywhere, Category = Settings, meta = (ClampMin = "0.0", ClampMax = "1.0", UIMin = "0.0", UIMax = "1.0"))
+	float Scattering4000Hz;
+
+	// How much this material scatters @ 8000Hz
+	UPROPERTY(VisibleAnywhere, Category = Settings, meta = (ClampMin = "0.0", ClampMax = "1.0", UIMin = "0.0", UIMax = "1.0"))
+	float Scattering8000Hz;
+};
\ No newline at end of file
diff --git a/Engine/Plugins/Runtime/Oculus/OculusAudio/Source/OculusAudio/Private/OculusAudioMixer.cpp b/Engine/Plugins/Runtime/Oculus/OculusAudio/Source/OculusAudio/Private/OculusAudioMixer.cpp
new file mode 100644
--- /dev/null
+++ b/Engine/Plugins/Runtime/Oculus/OculusAudio/Source/OculusAudio/Private/OculusAudioMixer.cpp
@@ -0,0 +1,188 @@
+// Copyright Epic Games, Inc. All Rights Reserved.
+
+#include "OculusAudioMixer.h"
+#include "OculusAudioSettings.h"
+#include "OculusAudioSourceSettings.h"
+#include "OculusAudioContextManager.h"
+#include "IOculusAudioPlugin.h"
+#include "Stats/Stats.h"
+
+
+float dbToLinear(float db)
+{
+	return powf(10.0f, db / 20.0f);
+}
+
+OculusAudioSpatializationAudioMixer::OculusAudioSpatializationAudioMixer()
+	: Context(nullptr)
+{
+}
+
+OculusAudioSpatializationAudioMixer::~OculusAudioSpatializationAudioMixer()
+{
+}
+
+void OculusAudioSpatializationAudioMixer::ClearContext()
+{
+	Context = nullptr;
+}
+
+void OculusAudioSpatializationAudioMixer::Initialize(const FAudioPluginInitializationParams InitializationParams)
+{
+	FScopeLock ScopeLock(&ContextLock);
+	InitParams = InitializationParams;
+
+	Context = FOculusAudioContextManager::GetContextForAudioDevice(InitializationParams.AudioDevicePtr);
+	if (!Context)
+	{
+		Context = FOculusAudioContextManager::CreateContextForAudioDevice(InitializationParams.AudioDevicePtr);
+		check(Context);
+	}
+
+	Params.AddDefaulted(InitParams.NumSources);
+
+	// UE is in centimeters
+	ovrResult Result = OVRA_CALL(ovrAudio_SetUnitScale)(Context, 0.01f);
+	OVR_AUDIO_CHECK(Result == ovrSuccess, "Failed to set unit scale");
+
+	const UOculusAudioSettings* Settings = GetDefault<UOculusAudioSettings>();
+	ApplyOculusAudioSettings(Settings);
+
+	TickDelegateHandle = FTSTicker::GetCoreTicker().AddTicker(FTickerDelegate::CreateRaw(this, &OculusAudioSpatializationAudioMixer::Tick));
+}
+
+void OculusAudioSpatializationAudioMixer::ApplyOculusAudioSettings(const UOculusAudioSettings* Settings)
+{
+	FScopeLock ScopeLock(&ContextLock);
+	ovrResult Result = OVRA_CALL(ovrAudio_Enable)(Context, ovrAudioEnable_SimpleRoomModeling, Settings->EarlyReflections);
+	OVR_AUDIO_CHECK(Result, "Failed to enable reflections");
+
+	Result = OVRA_CALL(ovrAudio_Enable)(Context, ovrAudioEnable_LateReverberation, Settings->LateReverberation);
+	OVR_AUDIO_CHECK(Result, "Failed to enable reverb");
+
+	ovrAudioBoxRoomParameters Room = { 0 };
+	Room.brp_Size = sizeof(Room);
+	Room.brp_Width = Settings->Width;
+	Room.brp_Height = Settings->Height;
+	Room.brp_Depth = Settings->Depth;
+	Room.brp_ReflectLeft = Settings->ReflectionCoefLeft;
+	Room.brp_ReflectRight = Settings->ReflectionCoefRight;
+	Room.brp_ReflectUp = Settings->ReflectionCoefUp;
+	Room.brp_ReflectDown = Settings->ReflectionCoefDown;
+	Room.brp_ReflectBehind = Settings->ReflectionCoefBack;
+	Room.brp_ReflectFront = Settings->ReflectionCoefFront;
+
+	Result = OVRA_CALL(ovrAudio_SetSimpleBoxRoomParameters)(Context, &Room);
+	OVR_AUDIO_CHECK(Result, "Failed to set room parameters");
+
+	Result = OVRA_CALL(ovrAudio_SetSharedReverbWetLevel)(Context, dbToLinear(Settings->ReverbWetLevel));
+	OVR_AUDIO_CHECK(Result, "Failed to set room parameters");
+}
+
+void OculusAudioSpatializationAudioMixer::Shutdown()
+{
+	if (TickDelegateHandle.IsValid())
+	{
+		FTSTicker::GetCoreTicker().RemoveTicker(TickDelegateHandle);
+	}
+}
+
+bool OculusAudioSpatializationAudioMixer::IsSpatializationEffectInitialized() const
+{
+	return Context != nullptr;
+}
+
+void OculusAudioSpatializationAudioMixer::OnInitSource(const uint32 SourceId, const FName& AudioComponentUserId, USpatializationPluginSourceSettingsBase* InSettings)
+{
+	FScopeLock ScopeLock(&ContextLock);
+	if (Context == nullptr)
+	{
+		UE_LOG(LogAudio, Error, TEXT("Oculus Audio Error - Context uninitialized. Sound %s (id=%d) will not play!"), *AudioComponentUserId.ToString(), SourceId);
+		return;
+	}
+
+	if (InSettings != nullptr) 
+	{
+		UOculusAudioSourceSettings* Settings = CastChecked<UOculusAudioSourceSettings>(InSettings);
+
+		uint32 Flags = 0;
+		if (!Settings->EarlyReflectionsEnabled)
+			Flags |= ovrAudioSourceFlag_ReflectionsDisabled;
+
+		ovrResult Result = OVRA_CALL(ovrAudio_SetAudioSourceFlags)(Context, SourceId, Flags);
+		OVR_AUDIO_CHECK(Result, "Failed to set audio source flags");
+
+		ovrAudioSourceAttenuationMode mode = Settings->AttenuationEnabled ? ovrAudioSourceAttenuationMode_InverseSquare : ovrAudioSourceAttenuationMode_None;
+		Result = OVRA_CALL(ovrAudio_SetAudioSourceAttenuationMode)(Context, SourceId, mode, 1.0f);
+		OVR_AUDIO_CHECK(Result, "Failed to set audio source attenuation mode");
+
+		Result = OVRA_CALL(ovrAudio_SetAudioSourceRange)(Context, SourceId, Settings->AttenuationRangeMinimum, Settings->AttenuationRangeMaximum);
+		OVR_AUDIO_CHECK(Result, "Failed to set audio source attenuation range");
+
+		Result = OVRA_CALL(ovrAudio_SetAudioSourceRadius)(Context, SourceId, Settings->VolumetricRadius);
+		OVR_AUDIO_CHECK(Result, "Failed to set audio source volumetric radius");
+
+		Result = OVRA_CALL(ovrAudio_SetAudioReverbSendLevel)(Context, SourceId, dbToLinear(Settings->ReverbSendLevel));
+		OVR_AUDIO_CHECK(Result, "Failed to set reverb send level");
+	}
+}
+
+void OculusAudioSpatializationAudioMixer::SetSpatializationParameters(uint32 VoiceId, const FSpatializationParams& InParams)
+{
+	Params[VoiceId] = InParams;
+}
+
+void OculusAudioSpatializationAudioMixer::ProcessAudio(const FAudioPluginSourceInputData& InputData, FAudioPluginSourceOutputData& OutputData)
+{
+	FScopeLock ScopeLock(&ContextLock);
+	if (InputData.SpatializationParams && Context)
+	{
+		Params[InputData.SourceId] = *InputData.SpatializationParams;
+
+		// Translate the input position to OVR coordinates
+		FVector OvrListenerPosition = ToOVRVector(Params[InputData.SourceId].ListenerPosition);
+		FVector OvrListenerForward = ToOVRVector(Params[InputData.SourceId].ListenerOrientation.GetForwardVector());
+		FVector OvrListenerUp = ToOVRVector(Params[InputData.SourceId].ListenerOrientation.GetUpVector());
+
+		ovrResult Result = OVRA_CALL(ovrAudio_SetListenerVectors)(Context,
+			OvrListenerPosition.X, OvrListenerPosition.Y, OvrListenerPosition.Z,
+			OvrListenerForward.X, OvrListenerForward.Y, OvrListenerForward.Z,
+			OvrListenerUp.X, OvrListenerUp.Y, OvrListenerUp.Z);
+		OVR_AUDIO_CHECK(Result, "Failed to set listener position and rotation");
+
+		// Translate the input position to OVR coordinates
+		FVector OvrPosition = ToOVRVector(Params[InputData.SourceId].EmitterWorldPosition);
+
+		// Set the source position to current audio position
+		Result = OVRA_CALL(ovrAudio_SetAudioSourcePos)(Context, InputData.SourceId, OvrPosition.X, OvrPosition.Y, OvrPosition.Z);
+		OVR_AUDIO_CHECK(Result, "Failed to set audio source position");
+
+		// Perform the processing
+		uint32 Status;
+		Result = OVRA_CALL(ovrAudio_SpatializeMonoSourceInterleaved)(Context, InputData.SourceId, &Status, OutputData.AudioBuffer.GetData(), InputData.AudioBuffer->GetData());
+		OVR_AUDIO_CHECK(Result, "Failed to spatialize mono source interleaved");
+	}
+}
+
+bool OculusAudioSpatializationAudioMixer::Tick(float DeltaTime)
+{
+	QUICK_SCOPE_CYCLE_COUNTER(STAT_OculusAudioSpatializationAudioMixer_Tick);
+	if (ContextLock.TryLock())
+	{
+		if (Context != nullptr)
+		{
+			ovrResult Result = OVRA_CALL(ovrAudio_UpdateRoomModel)(Context, 1.0f);
+
+			UOculusAudioSettings* settings = GetMutableDefault<UOculusAudioSettings>();
+			Result = OVRA_CALL(ovrAudio_SetPropagationQuality)(Context, settings->PropagationQuality);
+			if (Result != ovrSuccess)
+			{
+				UE_LOG(LogTemp, Warning, TEXT("Bad Propagation Quality setting %f!"), settings->PropagationQuality);
+			}
+		}
+
+		ContextLock.Unlock();
+	}
+	
+	return true;
+}
diff --git a/Engine/Plugins/Runtime/Oculus/OculusAudio/Source/OculusAudio/Private/OculusAudioMixer.h b/Engine/Plugins/Runtime/Oculus/OculusAudio/Source/OculusAudio/Private/OculusAudioMixer.h
new file mode 100644
--- /dev/null
+++ b/Engine/Plugins/Runtime/Oculus/OculusAudio/Source/OculusAudio/Private/OculusAudioMixer.h
@@ -0,0 +1,65 @@
+// Copyright Epic Games, Inc. All Rights Reserved.
+
+#pragma once
+#include "OVR_Audio.h"
+#include "Misc/Paths.h"
+#include "OculusAudioDllManager.h"
+#include "OculusAudioSettings.h"
+#include "Containers/Ticker.h"
+#include "AudioMixerDevice.h"
+#include "Containers/Ticker.h"
+
+/************************************************************************/
+/* OculusAudioSpatializationAudioMixer									*/
+/* This implementation of IAudioSpatialization uses the Oculus Audio	*/
+/* library to render audio sources with HRTF spatialization.			*/
+/************************************************************************/
+class OculusAudioSpatializationAudioMixer : public IAudioSpatialization
+{
+public:
+	OculusAudioSpatializationAudioMixer();
+	~OculusAudioSpatializationAudioMixer();
+
+	void ClearContext();
+
+	//~ Begin IAudioSpatialization 
+	virtual void Initialize(const FAudioPluginInitializationParams InitializationParams) override;
+	virtual void Shutdown() override;
+	virtual bool IsSpatializationEffectInitialized() const override;
+	virtual void OnInitSource(const uint32 SourceId, const FName& AudioComponentUserId, USpatializationPluginSourceSettingsBase* InSettings) override;
+	virtual void SetSpatializationParameters(uint32 VoiceId, const FSpatializationParams& InParams) override;
+	virtual void ProcessAudio(const FAudioPluginSourceInputData& InputData, FAudioPluginSourceOutputData& OutputData) override;
+	//~ End IAudioSpatialization
+	
+	bool Tick(float DeltaTime);
+
+	// Helper function to convert from UE coords to OVR coords.
+	static FVector ToOVRVector(const FVector& InVec)
+	{
+		return FVector(InVec.Y, InVec.Z, -InVec.X);
+	}
+
+	static FVector ToOVRVector(const Audio::FChannelPositionInfo& ChannelPositionInfo)
+	{
+		FVector OvrVector;
+		OvrVector.X = ChannelPositionInfo.Radius * FMath::Sin(ChannelPositionInfo.Azimuth) * FMath::Cos(ChannelPositionInfo.Elevation);
+		OvrVector.Y = ChannelPositionInfo.Radius * FMath::Sin(ChannelPositionInfo.Azimuth) * FMath::Sin(ChannelPositionInfo.Elevation);
+		OvrVector.Z = ChannelPositionInfo.Radius * FMath::Cos(ChannelPositionInfo.Azimuth);
+
+		return OvrVector;
+	}
+
+	static const int MIXER_CLASS_ID = 0x98765432;
+	const int ClassID = MIXER_CLASS_ID;
+private:
+
+	void ApplyOculusAudioSettings(const UOculusAudioSettings* Settings);
+
+	TArray<FSpatializationParams> Params;
+
+	ovrAudioContext Context;
+	FCriticalSection ContextLock;
+
+	FAudioPluginInitializationParams InitParams;
+	FTSTicker::FDelegateHandle TickDelegateHandle;
+};
\ No newline at end of file
diff --git a/Engine/Plugins/Runtime/Oculus/OculusAudio/Source/OculusAudio/Private/OculusAudioReverb.cpp b/Engine/Plugins/Runtime/Oculus/OculusAudio/Source/OculusAudio/Private/OculusAudioReverb.cpp
new file mode 100644
--- /dev/null
+++ b/Engine/Plugins/Runtime/Oculus/OculusAudio/Source/OculusAudio/Private/OculusAudioReverb.cpp
@@ -0,0 +1,127 @@
+// Copyright Epic Games, Inc. All Rights Reserved.
+
+#include "OculusAudioReverb.h"
+#include "OculusAudioMixer.h"
+#include "OculusAudioSettings.h"
+#include "OculusAudioContextManager.h"
+
+#include "Sound/SoundSubmix.h"
+#include "Engine/Engine.h"
+
+
+namespace
+{
+	TSharedPtr<FSubmixEffectOculusReverbPlugin, ESPMode::ThreadSafe> CastEffectToPluginSharedPtr(FSoundEffectSubmixPtr InSubmixEffect)
+	{
+		return StaticCastSharedPtr<FSubmixEffectOculusReverbPlugin, FSoundEffectSubmix, ESPMode::ThreadSafe>(InSubmixEffect);
+	}
+} // namespace <>
+
+
+void FSubmixEffectOculusReverbPlugin::OnNewDeviceCreated(Audio::FDeviceId DeviceId)
+{
+	FAudioDeviceManagerDelegates::OnAudioDeviceCreated.Remove(DeviceCreatedHandle);
+	DeviceCreatedHandle.Reset();
+
+	if (GEngine)
+	{
+		FAudioDevice* AudioDevice = GEngine->GetAudioDeviceManager()->GetAudioDeviceRaw(DeviceId);
+
+		if (AudioDevice)
+		{
+			Context = FOculusAudioContextManager::GetContextForAudioDevice(AudioDevice);
+
+			if (!Context)
+			{
+				Context = FOculusAudioContextManager::CreateContextForAudioDevice(AudioDevice);
+			}
+		}
+	}
+}
+
+void FSubmixEffectOculusReverbPlugin::ClearContext()
+{
+	Context = nullptr;
+}
+
+void FSubmixEffectOculusReverbPlugin::Init(const FSoundEffectSubmixInitData& InInitData)
+{
+	DeviceCreatedHandle = FAudioDeviceManagerDelegates::OnAudioDeviceCreated.AddRaw(this, &FSubmixEffectOculusReverbPlugin::OnNewDeviceCreated);
+}
+
+FSubmixEffectOculusReverbPlugin::FSubmixEffectOculusReverbPlugin()
+	: Context(nullptr)
+{
+}
+
+FSubmixEffectOculusReverbPlugin::~FSubmixEffectOculusReverbPlugin()
+{
+	if (DeviceCreatedHandle.IsValid())
+	{
+		FAudioDeviceManagerDelegates::OnAudioDeviceCreated.Remove(DeviceCreatedHandle);
+		DeviceCreatedHandle.Reset();
+	}
+}
+
+void FSubmixEffectOculusReverbPlugin::OnProcessAudio(const FSoundEffectSubmixInputData& InputData, FSoundEffectSubmixOutputData& OutputData)
+{
+	if (Context)
+	{
+		int Enabled = 0;
+		ovrResult Result = OVRA_CALL(ovrAudio_IsEnabled)(Context, ovrAudioEnable_LateReverberation, &Enabled);
+		OVR_AUDIO_CHECK(Result, "Failed to check if reverb is Enabled");
+
+		if (Enabled != 0)
+		{
+			uint32_t Status = 0;
+			Result = OVRA_CALL(ovrAudio_MixInSharedReverbInterleaved)(Context, &Status, OutputData.AudioBuffer->GetData());
+			OVR_AUDIO_CHECK(Result, "Failed to process reverb");
+		}
+	}
+}
+
+void OculusAudioReverb::ClearContext()
+{
+	Context = nullptr;
+	if (SubmixEffect.IsValid())
+	{
+		CastEffectToPluginSharedPtr(SubmixEffect)->ClearContext();
+	}
+}
+
+FSoundEffectSubmixPtr OculusAudioReverb::GetEffectSubmix()
+{
+	if (!SubmixEffect.IsValid())
+	{
+		if (!ReverbPreset)
+		{
+			ReverbPreset = NewObject<USubmixEffectOculusReverbPluginPreset>();
+			ReverbPreset->AddToRoot();
+		}
+
+		SubmixEffect = USoundEffectPreset::CreateInstance<FSoundEffectSubmixInitData, FSoundEffectSubmix>(FSoundEffectSubmixInitData(), *ReverbPreset);
+		SubmixEffect->SetEnabled(true);
+	}
+
+	return SubmixEffect;
+}
+
+USoundSubmix* OculusAudioReverb::GetSubmix()
+{
+	const UOculusAudioSettings* Settings = GetDefault<UOculusAudioSettings>();
+	check(Settings);
+	USoundSubmix* ReverbSubmix = Cast<USoundSubmix>(Settings->OutputSubmix.TryLoad());
+	if (!ReverbSubmix)
+	{
+		static const FString DefaultSubmixName = TEXT("Oculus Reverb Submix");
+		UE_LOG(LogAudio, Error, TEXT("Failed to load Oculus Reverb Submix from object path '%s' in OculusSettings. Creating '%s' as stub."),
+			*Settings->OutputSubmix.GetAssetPathString(),
+			*DefaultSubmixName);
+
+		ReverbSubmix = NewObject<USoundSubmix>(USoundSubmix::StaticClass(), *DefaultSubmixName);
+		ReverbSubmix->bMuteWhenBackgrounded = true;
+	}
+	ReverbSubmix->bAutoDisable = false;
+
+	return ReverbSubmix;
+}
\ No newline at end of file
diff --git a/Engine/Plugins/Runtime/Oculus/OculusAudio/Source/OculusAudio/Private/OculusAudioReverb.h b/Engine/Plugins/Runtime/Oculus/OculusAudio/Source/OculusAudio/Private/OculusAudioReverb.h
new file mode 100644
--- /dev/null
+++ b/Engine/Plugins/Runtime/Oculus/OculusAudio/Source/OculusAudio/Private/OculusAudioReverb.h
@@ -0,0 +1,111 @@
+// Copyright Epic Games, Inc. All Rights Reserved.
+
+#pragma once
+
+#include "AudioDevice.h"
+#include "IAudioExtensionPlugin.h"
+#include "OVR_Audio.h"
+#include "Sound/SoundEffectBase.h"
+#include "Sound/SoundEffectSubmix.h"
+#include "Templates/UniquePtr.h"
+#include "UObject/Object.h"
+#include "UObject/ObjectMacros.h"
+
+#include "OculusAudioReverb.generated.h"
+
+
+// Forward Declarations
+class USubmixEffectOculusReverbPluginPreset;
+
+class FSubmixEffectOculusReverbPlugin : public FSoundEffectSubmix
+{
+	virtual void Init(const FSoundEffectSubmixInitData& InInitData) override;
+	virtual uint32 GetDesiredInputChannelCountOverride() const override
+	{
+		static const int STEREO = 2;
+		return STEREO; // PAS
+	}
+	virtual void OnProcessAudio(const FSoundEffectSubmixInputData& InData, FSoundEffectSubmixOutputData& OutData) override;
+	virtual void OnPresetChanged() override
+	{
+		return; // PAS
+	}
+
+public:
+	FSubmixEffectOculusReverbPlugin();
+	~FSubmixEffectOculusReverbPlugin();
+
+	void ClearContext();
+
+private:
+	void OnNewDeviceCreated(Audio::FDeviceId InDeviceId);
+
+	FDelegateHandle DeviceCreatedHandle;
+	ovrAudioContext Context;
+};
+
+/************************************************************************/
+/* OculusAudioReverb													*/
+/* This implementation of IAudioReverb uses the Oculus Audio			*/
+/* library to render spatial reverb.									*/
+/************************************************************************/
+class OculusAudioReverb : public IAudioReverb
+{
+public:
+	OculusAudioReverb()
+		: Context(nullptr)
+		, ReverbPreset(nullptr)
+	{
+		// empty
+	}
+
+	void ClearContext();
+
+	virtual void OnInitSource(const uint32 SourceId, const FName& AudioComponentUserId, const uint32 NumChannels, UReverbPluginSourceSettingsBase* InSettings) override
+	{
+		return; // PAS
+	}
+
+	virtual void OnReleaseSource(const uint32 SourceId) override 
+	{
+		return; // PAS
+	}
+
+	virtual FSoundEffectSubmixPtr GetEffectSubmix() override;
+
+	virtual USoundSubmix* GetSubmix() override;
+
+	virtual void ProcessSourceAudio(const FAudioPluginSourceInputData& InputData, FAudioPluginSourceOutputData& OutputData) override
+	{
+		return; // PAS
+	}
+private:
+	ovrAudioContext* Context;
+	TSoundEffectSubmixPtr SubmixEffect;
+	USubmixEffectOculusReverbPluginPreset* ReverbPreset;
+};
+
+USTRUCT()
+struct OCULUSAUDIO_API FSubmixEffectOculusReverbPluginSettings
+{
+	GENERATED_USTRUCT_BODY()
+
+	FSubmixEffectOculusReverbPluginSettings() = default;
+};
+
+UCLASS()
+class OCULUSAUDIO_API USubmixEffectOculusReverbPluginPreset : public USoundEffectSubmixPreset
+{
+	GENERATED_BODY()
+
+public:
+	EFFECT_PRESET_METHODS(SubmixEffectOculusReverbPlugin)
+
+	UFUNCTION()
+	void SetSettings(const FSubmixEffectOculusReverbPluginSettings& InSettings)
+	{
+	}
+
+	UPROPERTY()
+	FSubmixEffectOculusReverbPluginSettings Settings;
+};
\ No newline at end of file
diff --git a/Engine/Plugins/Runtime/Oculus/OculusAudio/Source/OculusAudio/Private/OculusAudioSettings.cpp b/Engine/Plugins/Runtime/Oculus/OculusAudio/Source/OculusAudio/Private/OculusAudioSettings.cpp
new file mode 100644
--- /dev/null
+++ b/Engine/Plugins/Runtime/Oculus/OculusAudio/Source/OculusAudio/Private/OculusAudioSettings.cpp
@@ -0,0 +1,37 @@
+// Copyright Epic Games, Inc. All Rights Reserved.
+
+#include "OculusAudioSettings.h"
+
+UOculusAudioSettings::UOculusAudioSettings() :
+	ReverbWetLevel(0.0f)
+	, EarlyReflections(true)
+	, LateReverberation(false)
+	, PropagationQuality(1.0f)
+	, Width(8.0f)
+	, Height(3.0f)
+	, Depth(5.0f)
+	, ReflectionCoefRight(0.25f)
+	, ReflectionCoefLeft(0.25f)
+	, ReflectionCoefUp(0.5f)
+	, ReflectionCoefDown(0.1f)
+	, ReflectionCoefBack(0.25f)
+	, ReflectionCoefFront(0.25f)
+{
+	OutputSubmix = FString(TEXT("/OculusAudio/OculusSubmixDefault.OculusSubmixDefault"));
+}
+
+#if WITH_EDITOR
+void UOculusAudioSettings::PostEditChangeProperty(FPropertyChangedEvent& PropertyChangedEvent)
+{
+	// TODO notify
+
+	Super::PostEditChangeProperty(PropertyChangedEvent);
+}
+
+bool UOculusAudioSettings::CanEditChange(const FProperty* InProperty) const
+{
+	// TODO disable settings when reflection engine is disabled
+
+	return Super::CanEditChange(InProperty);
+}
+#endif
\ No newline at end of file
diff --git a/Engine/Plugins/Runtime/Oculus/OculusAudio/Source/OculusAudio/Private/OculusAudioSettings.h b/Engine/Plugins/Runtime/Oculus/OculusAudio/Source/OculusAudio/Private/OculusAudioSettings.h
new file mode 100644
--- /dev/null
+++ b/Engine/Plugins/Runtime/Oculus/OculusAudio/Source/OculusAudio/Private/OculusAudioSettings.h
@@ -0,0 +1,63 @@
+// Copyright Epic Games, Inc. All Rights Reserved.
+
+#pragma once
+
+#include "OculusAudioSettings.generated.h"
+
+
+UCLASS(config = Engine, defaultconfig)
+class OCULUSAUDIO_API UOculusAudioSettings : public UObject
+{
+	GENERATED_BODY()
+
+public:
+
+	UOculusAudioSettings();
+
+#if WITH_EDITOR
+	virtual void PostEditChangeProperty(struct FPropertyChangedEvent& PropertyChangedEvent) override;
+	virtual bool CanEditChange(const FProperty* InProperty) const override;
+#endif
+
+	UPROPERTY(GlobalConfig, EditAnywhere, Category = "Reverb", meta = (AllowedClasses = "/Script/Engine.SoundSubmix"))
+	FSoftObjectPath OutputSubmix;
+
+	UPROPERTY(GlobalConfig, EditAnywhere, Category = "Reverb", meta = (ClampMin = "-60.0", ClampMax = "20.0", UIMin = "-60.0", UIMax = "20.0"))
+	float ReverbWetLevel;
+
+	UPROPERTY(GlobalConfig, EditAnywhere, Category = "Room Modeling")
+	bool EarlyReflections;
+
+	UPROPERTY(GlobalConfig, EditAnywhere, Category = "Room Modeling")
+	bool LateReverberation;
+
+	UPROPERTY(GlobalConfig, EditAnywhere, Category = "Propagation Quality", meta = (ClampMin = "0.0", ClampMax = "2.0", UIMin = "0.0", UIMax = "2.0"))
+	float PropagationQuality;
+	
+	UPROPERTY(GlobalConfig, EditAnywhere, Category = "Room Modeling", meta = (ClampMin = "0.0", ClampMax = "200.0", UIMin = "0.0", UIMax = "200.0"))
+	float Width;
+	
+	UPROPERTY(GlobalConfig, EditAnywhere, Category = "Room Modeling", meta = (ClampMin = "0.0", ClampMax = "200.0", UIMin = "0.0", UIMax = "200.0"))
+	float Height;
+
+	UPROPERTY(GlobalConfig, EditAnywhere, Category = "Room Modeling", meta = (ClampMin = "0.0", ClampMax = "200.0", UIMin = "0.0", UIMax = "200.0"))
+	float Depth;
+
+	UPROPERTY(GlobalConfig, EditAnywhere, Category = "Room Modeling", meta = (ClampMin = "0.0", ClampMax = "0.97", UIMin = "0.0", UIMax = "0.97"))
+	float ReflectionCoefRight;
+
+	UPROPERTY(GlobalConfig, EditAnywhere, Category = "Room Modeling", meta = (ClampMin = "0.0", ClampMax = "0.97", UIMin = "0.0", UIMax = "0.97"))
+	float ReflectionCoefLeft;
+	
+	UPROPERTY(GlobalConfig, EditAnywhere, Category = "Room Modeling", meta = (ClampMin = "0.0", ClampMax = "0.97", UIMin = "0.0", UIMax = "0.97"))
+	float ReflectionCoefUp;
+	
+	UPROPERTY(GlobalConfig, EditAnywhere, Category = "Room Modeling", meta = (ClampMin = "0.0", ClampMax = "0.97", UIMin = "0.0", UIMax = "0.97"))
+	float ReflectionCoefDown;
+
+	UPROPERTY(GlobalConfig, EditAnywhere, Category = "Room Modeling", meta = (ClampMin = "0.0", ClampMax = "0.97", UIMin = "0.0", UIMax = "0.97"))
+	float ReflectionCoefBack;
+
+	UPROPERTY(GlobalConfig, EditAnywhere, Category = "Room Modeling", meta = (ClampMin = "0.0", ClampMax = "0.97", UIMin = "0.0", UIMax = "0.97"))
+	float ReflectionCoefFront;
+};
\ No newline at end of file
diff --git a/Engine/Plugins/Runtime/Oculus/OculusAudio/Source/OculusAudio/Public/IOculusAudioPlugin.h b/Engine/Plugins/Runtime/Oculus/OculusAudio/Source/OculusAudio/Public/IOculusAudioPlugin.h
new file mode 100644
--- /dev/null
+++ b/Engine/Plugins/Runtime/Oculus/OculusAudio/Source/OculusAudio/Public/IOculusAudioPlugin.h
@@ -0,0 +1,34 @@
+// Copyright Epic Games, Inc. All Rights Reserved.
+
+#pragma once
+
+#include "Modules/ModuleManager.h"
+#include "OculusAudio.h"
+#include "OculusAudioDllManager.h"
+#include "OculusAmbisonicSpatializer.h"
+
+/************************************************************************/
+/* FOculusAudioPlugin												   */
+/* Module interface. Also handles loading and unloading the Oculus Audio*/
+/* DLL using FOculusAudioDllManager.									*/
+/************************************************************************/
+class FOculusAudioPlugin : public IModuleInterface
+{
+public:
+	//~ Begin IModuleInterface 
+	virtual void StartupModule() override;
+	virtual void ShutdownModule() override;
+	//~ End IModuleInterface
+
+	void RegisterAudioDevice(FAudioDevice* AudioDeviceHandle);
+	void UnregisterAudioDevice(FAudioDevice* AudioDeviceHandle);
+
+	FOculusSpatializationPluginFactory* GetSpatializationPluginFactory() { return &PluginFactory; }
+	FOculusReverbPluginFactory* GetReverbPluginFactory() { return &ReverbPluginFactory; }
+
+private:
+	TArray<FAudioDevice*> RegisteredAudioDevices;
+	FOculusSpatializationPluginFactory PluginFactory;
+	FOculusReverbPluginFactory ReverbPluginFactory;
+	FOculusAmbisonicsFactory AmbisonicsFactory;
+};
diff --git a/Engine/Plugins/Runtime/Oculus/OculusAudio/Source/OculusAudio/Public/OculusAmbisonicSpatializer.h b/Engine/Plugins/Runtime/Oculus/OculusAudio/Source/OculusAudio/Public/OculusAmbisonicSpatializer.h
new file mode 100644
--- /dev/null
+++ b/Engine/Plugins/Runtime/Oculus/OculusAudio/Source/OculusAudio/Public/OculusAmbisonicSpatializer.h
@@ -0,0 +1,34 @@
+// Copyright Epic Games, Inc. All Rights Reserved.
+
+#pragma once
+
+#include "IAudioExtensionPlugin.h"
+#include "AudioEffect.h"
+#include "DSP/Dsp.h"
+#include "Sound/SoundEffectSubmix.h"
+#include "OVR_Audio.h"
+#include "OculusAmbisonicsSettings.h"
+
+class FOculusAmbisonicsFactory : public ISoundfieldFactory
+{
+public:
+	FOculusAmbisonicsFactory();
+	virtual ~FOculusAmbisonicsFactory();
+
+	// Begin ISoundfieldFactory
+	virtual const USoundfieldEncodingSettingsBase* GetDefaultEncodingSettings() override;
+	virtual FName GetSoundfieldFormatName() override;
+	virtual TUniquePtr<ISoundfieldEncoderStream> CreateEncoderStream(const FAudioPluginInitializationParams& InitInfo, const ISoundfieldEncodingSettingsProxy& InitialSettings) override;
+	virtual TUniquePtr<ISoundfieldDecoderStream> CreateDecoderStream(const FAudioPluginInitializationParams& InitInfo, const ISoundfieldEncodingSettingsProxy& InitialSettings) override;
+	virtual TUniquePtr<ISoundfieldTranscodeStream> CreateTranscoderStream(const FName SourceFormat, const ISoundfieldEncodingSettingsProxy& InitialSourceSettings, const FName DestinationFormat, const ISoundfieldEncodingSettingsProxy& InitialDestinationSettings, const FAudioPluginInitializationParams& InitInfo) override;
+	virtual TUniquePtr<ISoundfieldMixerStream> CreateMixerStream(const ISoundfieldEncodingSettingsProxy& InitialSettings) override;
+	virtual TUniquePtr<ISoundfieldAudioPacket> CreateEmptyPacket() override;
+	virtual bool IsTranscodeRequiredBetweenSettings(const ISoundfieldEncodingSettingsProxy& SourceSettings, const ISoundfieldEncodingSettingsProxy& DestinationSettings) override;
+	virtual bool CanTranscodeFromSoundfieldFormat(FName SourceFormat, const ISoundfieldEncodingSettingsProxy& SourceEncodingSettings) override;
+	virtual bool CanTranscodeToSoundfieldFormat(FName DestinationFormat, const ISoundfieldEncodingSettingsProxy& DestinationEncodingSettings) override;
+	virtual UClass* GetCustomEncodingSettingsClass() const override;
+	//~ ISoundfieldFactory
+
+private:
+	static TMap<FAudioDevice*, ovrAudioContext> ContextMap;
+};
\ No newline at end of file
diff --git a/Engine/Plugins/Runtime/Oculus/OculusAudio/Source/OculusAudio/Public/OculusAmbisonicsSettings.h b/Engine/Plugins/Runtime/Oculus/OculusAudio/Source/OculusAudio/Public/OculusAmbisonicsSettings.h
new file mode 100644
--- /dev/null
+++ b/Engine/Plugins/Runtime/Oculus/OculusAudio/Source/OculusAudio/Public/OculusAmbisonicsSettings.h
@@ -0,0 +1,72 @@
+// Copyright Epic Games, Inc. All Rights Reserved.
+
+#pragma once
+
+#include "IAudioExtensionPlugin.h"
+
+#include "OculusAmbisonicsSettings.generated.h"
+
+
+UENUM()
+enum class EOculusAudioAmbisonicsMode : uint8
+{
+	// High quality ambisonic spatialization method
+	SphericalHarmonics,
+
+	// Alternative ambisonic spatialization method
+	VirtualSpeakers,
+};
+
+USTRUCT(BlueprintType)
+struct OCULUSAUDIO_API FSubmixEffectOculusAmbisonicSpatializerSettings
+{
+	GENERATED_USTRUCT_BODY()
+
+		// Ambisonic spatialization mode
+	UPROPERTY(EditAnywhere, BlueprintReadWrite, Category = Realtime)
+	EOculusAudioAmbisonicsMode AmbisonicMode;
+
+	FSubmixEffectOculusAmbisonicSpatializerSettings()
+		: AmbisonicMode(EOculusAudioAmbisonicsMode::SphericalHarmonics)
+	{
+	}
+};
+
+class FOculusAudioSoundfieldSettingsProxy : public ISoundfieldEncodingSettingsProxy
+{
+
+public:
+	const EOculusAudioAmbisonicsMode SpatailizationMode;
+
+	FOculusAudioSoundfieldSettingsProxy(EOculusAudioAmbisonicsMode InSpatializationMode)
+		: SpatailizationMode(InSpatializationMode)
+	{}
+
+	virtual uint32 GetUniqueId() const override
+	{
+		return (uint32)SpatailizationMode;
+	}
+
+
+	virtual TUniquePtr<ISoundfieldEncodingSettingsProxy> Duplicate() const override
+	{
+		return TUniquePtr<ISoundfieldEncodingSettingsProxy>(new FOculusAudioSoundfieldSettingsProxy(SpatailizationMode));
+	}
+
+};
+
+UCLASS()
+class OCULUSAUDIO_API UOculusAudioSoundfieldSettings : public USoundfieldEncodingSettingsBase
+{
+	GENERATED_BODY()
+
+public:
+	UPROPERTY(EditAnywhere, BlueprintReadOnly, Category = Ambisonics)
+	EOculusAudioAmbisonicsMode SpatializationMode;
+
+	virtual TUniquePtr<ISoundfieldEncodingSettingsProxy> GetProxy() const override
+	{
+		return TUniquePtr<ISoundfieldEncodingSettingsProxy>(new FOculusAudioSoundfieldSettingsProxy(SpatializationMode));
+	}
+};
+
diff --git a/Engine/Plugins/Runtime/Oculus/OculusAudio/Source/OculusAudio/Public/OculusAudioSourceSettings.cpp b/Engine/Plugins/Runtime/Oculus/OculusAudio/Source/OculusAudio/Public/OculusAudioSourceSettings.cpp
new file mode 100644
--- /dev/null
+++ b/Engine/Plugins/Runtime/Oculus/OculusAudio/Source/OculusAudio/Public/OculusAudioSourceSettings.cpp
@@ -0,0 +1,32 @@
+// Copyright Epic Games, Inc. All Rights Reserved.
+
+#include "OculusAudioSourceSettings.h"
+
+UOculusAudioSourceSettings::UOculusAudioSourceSettings() :
+	 EarlyReflectionsEnabled(true)
+	, AttenuationEnabled(true)
+	, AttenuationRangeMinimum(0.25f)
+	, AttenuationRangeMaximum(250.0f)
+	, VolumetricRadius(0.0f)
+	, ReverbSendLevel(0.0f)
+{
+
+}
+
+#if WITH_EDITOR
+bool UOculusAudioSourceSettings::CanEditChange(const FProperty* InProperty) const
+{
+	const bool ParentVal = Super::CanEditChange(InProperty);
+
+	if ((InProperty->GetFName() == GET_MEMBER_NAME_CHECKED(UOculusAudioSourceSettings, AttenuationRangeMinimum)))
+	{
+		return ParentVal && AttenuationEnabled;
+	}
+	if ((InProperty->GetFName() == GET_MEMBER_NAME_CHECKED(UOculusAudioSourceSettings, AttenuationRangeMaximum)))
+	{
+		return ParentVal && AttenuationEnabled;
+	}
+
+	return ParentVal;
+}
+#endif
\ No newline at end of file
diff --git a/Engine/Plugins/Runtime/Oculus/OculusAudio/Source/OculusAudio/Public/OculusAudioSourceSettings.h b/Engine/Plugins/Runtime/Oculus/OculusAudio/Source/OculusAudio/Public/OculusAudioSourceSettings.h
new file mode 100644
--- /dev/null
+++ b/Engine/Plugins/Runtime/Oculus/OculusAudio/Source/OculusAudio/Public/OculusAudioSourceSettings.h
@@ -0,0 +1,39 @@
+// Copyright Epic Games, Inc. All Rights Reserved.
+
+#pragma once
+
+#include "IAudioExtensionPlugin.h"
+
+#include "OculusAudioSourceSettings.generated.h"
+
+UCLASS()
+class OCULUSAUDIO_API UOculusAudioSourceSettings : public USpatializationPluginSourceSettingsBase
+{
+	GENERATED_BODY()
+
+public:
+	UOculusAudioSourceSettings();
+
+#if WITH_EDITOR
+	virtual bool CanEditChange(const FProperty* InProperty) const override;
+#endif
+
+	UPROPERTY(GlobalConfig, EditAnywhere, Category = "SpatializationSettings")
+	bool EarlyReflectionsEnabled;
+
+	UPROPERTY(GlobalConfig, EditAnywhere, Category = "SpatializationSettings")
+	bool AttenuationEnabled;
+
+	UPROPERTY(GlobalConfig, EditAnywhere, Category = "SpatializationSettings")
+	float AttenuationRangeMinimum;
+	
+	UPROPERTY(GlobalConfig, EditAnywhere, Category = "SpatializationSettings")
+	float AttenuationRangeMaximum;
+	
+	UPROPERTY(GlobalConfig, EditAnywhere, Category = "SpatializationSettings")
+	float VolumetricRadius;
+
+	UPROPERTY(GlobalConfig, EditAnywhere, Category = "SpatializationSettings", meta = (ClampMin = "-60.0", ClampMax = "20.0", UIMin = "-60.0", UIMax = "20.0"))
+	float ReverbSendLevel;
+};
+
diff --git a/Engine/Plugins/Runtime/Oculus/OculusAudio/Source/OculusAudioEditor/OculusAudioEditor.Build.cs b/Engine/Plugins/Runtime/Oculus/OculusAudio/Source/OculusAudioEditor/OculusAudioEditor.Build.cs
new file mode 100644
--- /dev/null
+++ b/Engine/Plugins/Runtime/Oculus/OculusAudio/Source/OculusAudioEditor/OculusAudioEditor.Build.cs
@@ -0,0 +1,49 @@
+// Copyright Epic Games, Inc. All Rights Reserved.
+
+namespace UnrealBuildTool.Rules
+{
+    public class OculusAudioEditor : ModuleRules
+    {
+        public OculusAudioEditor(ReadOnlyTargetRules Target) : base(Target)
+        {
+            PublicDependencyModuleNames.AddRange(
+                new string[] {
+                "Core",
+                "CoreUObject",
+                "Engine",
+                }
+                );
+
+            PrivateIncludePaths.AddRange(
+                new string[] {
+                    "OculusAudio/Private",
+ 					// ... add other private include paths required here ...
+				}
+                );
+
+            PrivateIncludePathModuleNames.AddRange(
+            new string[] {
+                "AssetTools",
+                "AudioEditor",
+                "OculusAudio",
+                "UnrealEd",
+                }
+                );
+            
+            PublicDependencyModuleNames.AddRange(
+                new string[] {
+                    "Core",
+                    "CoreUObject",
+                    "Engine",
+                    "InputCore",
+					"EditorFramework",
+                    "UnrealEd",
+                    "RHI",
+                    "AudioEditor",
+                    "AudioMixer",
+                    "OculusAudio"
+                }
+            );
+        }
+    }
+}
\ No newline at end of file
diff --git a/Engine/Plugins/Runtime/Oculus/OculusAudio/Source/OculusAudioEditor/Private/IOculusAudioEditorPlugin.cpp b/Engine/Plugins/Runtime/Oculus/OculusAudio/Source/OculusAudioEditor/Private/IOculusAudioEditorPlugin.cpp
new file mode 100644
--- /dev/null
+++ b/Engine/Plugins/Runtime/Oculus/OculusAudio/Source/OculusAudioEditor/Private/IOculusAudioEditorPlugin.cpp
@@ -0,0 +1,32 @@
+// Copyright Epic Games, Inc. All Rights Reserved.
+
+#include "IOculusAudioEditorPlugin.h"
+#include "OculusAudioSettings.h"
+#include "OculusAudioSourceSettingsFactory.h"
+#include "OculusAmbisonicSettingsFactory.h"
+
+#include "Modules/ModuleManager.h"
+#include "ISettingsModule.h"
+
+void FOculusAudioEditorPlugin::StartupModule()
+{
+    // Register asset types
+    IAssetTools& AssetTools = FModuleManager::LoadModuleChecked<FAssetToolsModule>("AssetTools").Get();
+    {
+		AssetTools.RegisterAssetTypeActions(MakeShareable(new FAssetTypeActions_OculusAudioSourceSettings));
+		AssetTools.RegisterAssetTypeActions(MakeShareable(new FAssetTypeActions_OculusAmbisonicsSettings));
+
+        ISettingsModule* SettingsModule = FModuleManager::Get().GetModulePtr<ISettingsModule>("Settings");
+        if (SettingsModule)
+        {
+            SettingsModule->RegisterSettings("Project", "Plugins", "Oculus Audio", NSLOCTEXT("OculusAudio", "Oculus Audio", "Oculus Audio"),
+                NSLOCTEXT("OculusAudio", "Configure Oculus Audio settings", "Configure Oculus Audio settings"), GetMutableDefault<UOculusAudioSettings>());
+        }
+    }
+}
+
+void FOculusAudioEditorPlugin::ShutdownModule()
+{
+}
+
+IMPLEMENT_MODULE(FOculusAudioEditorPlugin, OculusAudioEditor)
diff --git a/Engine/Plugins/Runtime/Oculus/OculusAudio/Source/OculusAudioEditor/Private/OculusAmbisonicSettingsFactory.cpp b/Engine/Plugins/Runtime/Oculus/OculusAudio/Source/OculusAudioEditor/Private/OculusAmbisonicSettingsFactory.cpp
new file mode 100644
--- /dev/null
+++ b/Engine/Plugins/Runtime/Oculus/OculusAudio/Source/OculusAudioEditor/Private/OculusAmbisonicSettingsFactory.cpp
@@ -0,0 +1,63 @@
+// Copyright Epic Games, Inc. All Rights Reserved.
+
+#include "OculusAmbisonicSettingsFactory.h"
+#include "AssetTypeCategories.h"
+#include "AudioAnalytics.h"
+#include "OculusAmbisonicsSettings.h"
+#include "UObject/ObjectMacros.h"
+#include "UObject/Object.h"
+
+#define LOCTEXT_NAMESPACE "AssetTypeActions"
+
+FText FAssetTypeActions_OculusAmbisonicsSettings::GetName() const
+{
+	return LOCTEXT("AssetTypeActions_OculusAmbisonicsSettings", "Oculus Ambisonics Settings");
+}
+
+FColor FAssetTypeActions_OculusAmbisonicsSettings::GetTypeColor() const
+{
+	return FColor(100, 100, 100);
+}
+
+const TArray<FText>& FAssetTypeActions_OculusAmbisonicsSettings::GetSubMenus() const
+{
+	static const TArray<FText> SubMenus
+	{
+		LOCTEXT("AssetSoundOculusSubMenu", "Oculus")
+	};
+
+	return SubMenus;
+}
+
+UClass* FAssetTypeActions_OculusAmbisonicsSettings::GetSupportedClass() const
+{
+	return UOculusAudioSoundfieldSettings::StaticClass();
+}
+
+uint32 FAssetTypeActions_OculusAmbisonicsSettings::GetCategories()
+{
+	return EAssetTypeCategories::Sounds;
+}
+
+UOculusAmbisonicsSettingsFactory::UOculusAmbisonicsSettingsFactory(const FObjectInitializer& ObjectInitializer)
+	: Super(ObjectInitializer)
+{
+	SupportedClass = UOculusAudioSoundfieldSettings::StaticClass();
+
+	bCreateNew = true;
+	bEditorImport = true;
+	bEditAfterNew = true;
+}
+
+UObject* UOculusAmbisonicsSettingsFactory::FactoryCreateNew(UClass* Class, UObject* InParent, FName Name, EObjectFlags Flags, UObject* Context,
+	FFeedbackContext* Warn)
+{
+	Audio::Analytics::RecordEvent_Usage(TEXT("OculusAudio.AmbisonicsSettingsCreated"));
+	return NewObject<UOculusAudioSoundfieldSettings>(InParent, Name, Flags);
+}
+
+uint32 UOculusAmbisonicsSettingsFactory::GetMenuCategories() const
+{
+	return EAssetTypeCategories::Sounds;
+}
+#undef LOCTEXT_NAMESPACE
diff --git a/Engine/Plugins/Runtime/Oculus/OculusAudio/Source/OculusAudioEditor/Private/OculusAmbisonicSettingsFactory.h b/Engine/Plugins/Runtime/Oculus/OculusAudio/Source/OculusAudioEditor/Private/OculusAmbisonicSettingsFactory.h
new file mode 100644
--- /dev/null
+++ b/Engine/Plugins/Runtime/Oculus/OculusAudio/Source/OculusAudioEditor/Private/OculusAmbisonicSettingsFactory.h
@@ -0,0 +1,31 @@
+// Copyright Epic Games, Inc. All Rights Reserved.
+
+#pragma once
+
+#include "IAudioExtensionPlugin.h"
+#include "Factories/Factory.h"
+#include "AssetTypeActions_Base.h"
+#include "AssetToolsModule.h"
+#include "OculusAmbisonicSettingsFactory.generated.h"
+
+class FAssetTypeActions_OculusAmbisonicsSettings : public FAssetTypeActions_Base
+{
+public:
+    virtual FText GetName() const override;
+    virtual FColor GetTypeColor() const override;
+	virtual const TArray<FText>& GetSubMenus() const override;
+    virtual UClass* GetSupportedClass() const override;
+    virtual uint32 GetCategories() override;
+};
+
+UCLASS(hidecategories = Object, MinimalAPI)
+class UOculusAmbisonicsSettingsFactory : public UFactory
+{
+	GENERATED_UCLASS_BODY()
+
+	virtual UObject* FactoryCreateNew(UClass* Class, UObject* InParent, FName Name, EObjectFlags Flags, UObject* Context,
+		FFeedbackContext* Warn) override;
+
+	virtual uint32 GetMenuCategories() const override;
+};
+
diff --git a/Engine/Plugins/Runtime/Oculus/OculusAudio/Source/OculusAudioEditor/Private/OculusAudioSourceSettingsFactory.cpp b/Engine/Plugins/Runtime/Oculus/OculusAudio/Source/OculusAudioEditor/Private/OculusAudioSourceSettingsFactory.cpp
new file mode 100644
--- /dev/null
+++ b/Engine/Plugins/Runtime/Oculus/OculusAudio/Source/OculusAudioEditor/Private/OculusAudioSourceSettingsFactory.cpp
@@ -0,0 +1,64 @@
+// Copyright Epic Games, Inc. All Rights Reserved.
+
+#include "OculusAudioSourceSettingsFactory.h"
+#include "AssetTypeCategories.h"
+#include "AudioAnalytics.h"
+#include "OculusAudioSourceSettings.h"
+#include "UObject/ObjectMacros.h"
+#include "UObject/Object.h"
+
+#define LOCTEXT_NAMESPACE "AssetTypeActions"
+
+FText FAssetTypeActions_OculusAudioSourceSettings::GetName() const
+{
+    return LOCTEXT("AssetTypeActions_OculusAudioSpatializationSettings", "Oculus Audio Source Settings");
+}
+
+const TArray<FText>& FAssetTypeActions_OculusAudioSourceSettings::GetSubMenus() const
+{
+	static const TArray<FText> SubMenus
+	{
+		LOCTEXT("AssetSoundOculusSubMenu", "Oculus")
+	};
+
+	return SubMenus;
+}
+
+FColor FAssetTypeActions_OculusAudioSourceSettings::GetTypeColor() const
+{
+    return FColor(100, 100, 100);
+}
+
+UClass* FAssetTypeActions_OculusAudioSourceSettings::GetSupportedClass() const
+{
+    return UOculusAudioSourceSettings::StaticClass();
+}
+
+uint32 FAssetTypeActions_OculusAudioSourceSettings::GetCategories()
+{
+    return EAssetTypeCategories::Sounds;
+}
+
+UOculusAudioSourceSettingsFactory::UOculusAudioSourceSettingsFactory(const FObjectInitializer& ObjectInitializer)
+    : Super(ObjectInitializer)
+{
+    SupportedClass = UOculusAudioSourceSettings::StaticClass();
+
+    bCreateNew = true;
+    bEditorImport = true;
+    bEditAfterNew = true;
+}
+
+UObject* UOculusAudioSourceSettingsFactory::FactoryCreateNew(UClass* Class, UObject* InParent, FName Name, EObjectFlags Flags, UObject* Context,
+    FFeedbackContext* Warn)
+{
+	Audio::Analytics::RecordEvent_Usage(TEXT("OculusAudio.SourceSettingsCreated"));
+    return NewObject<UOculusAudioSourceSettings>(InParent, Name, Flags);
+}
+
+uint32 UOculusAudioSourceSettingsFactory::GetMenuCategories() const
+{
+    return EAssetTypeCategories::Sounds;
+}
+
+#undef LOCTEXT_NAMESPACE
diff --git a/Engine/Plugins/Runtime/Oculus/OculusAudio/Source/OculusAudioEditor/Private/OculusAudioSourceSettingsFactory.h b/Engine/Plugins/Runtime/Oculus/OculusAudio/Source/OculusAudioEditor/Private/OculusAudioSourceSettingsFactory.h
new file mode 100644
--- /dev/null
+++ b/Engine/Plugins/Runtime/Oculus/OculusAudio/Source/OculusAudioEditor/Private/OculusAudioSourceSettingsFactory.h
@@ -0,0 +1,30 @@
+// Copyright Epic Games, Inc. All Rights Reserved.
+
+#pragma once
+
+#include "IAudioExtensionPlugin.h"
+#include "Factories/Factory.h"
+#include "AssetTypeActions_Base.h"
+#include "AssetToolsModule.h"
+#include "OculusAudioSourceSettingsFactory.generated.h"
+
+class FAssetTypeActions_OculusAudioSourceSettings : public FAssetTypeActions_Base
+{
+public:
+    virtual FText GetName() const override;
+    virtual FColor GetTypeColor() const override;
+	virtual const TArray<FText>& GetSubMenus() const override;
+	virtual UClass* GetSupportedClass() const override;
+    virtual uint32 GetCategories() override;
+};
+
+UCLASS(hidecategories = Object, MinimalAPI)
+class UOculusAudioSourceSettingsFactory : public UFactory
+{
+    GENERATED_UCLASS_BODY()
+
+	virtual UObject* FactoryCreateNew(UClass* Class, UObject* InParent, FName Name, EObjectFlags Flags, UObject* Context,
+		FFeedbackContext* Warn) override;
+
+	virtual uint32 GetMenuCategories() const override;
+};
diff --git a/Engine/Plugins/Runtime/Oculus/OculusAudio/Source/OculusAudioEditor/Public/IOculusAudioEditorPlugin.h b/Engine/Plugins/Runtime/Oculus/OculusAudio/Source/OculusAudioEditor/Public/IOculusAudioEditorPlugin.h
new file mode 100644
--- /dev/null
+++ b/Engine/Plugins/Runtime/Oculus/OculusAudio/Source/OculusAudioEditor/Public/IOculusAudioEditorPlugin.h
@@ -0,0 +1,13 @@
+// Copyright Epic Games, Inc. All Rights Reserved.
+
+#pragma once
+
+#include "Modules/ModuleInterface.h"
+#include "Modules/ModuleManager.h"
+
+class FOculusAudioEditorPlugin : public IModuleInterface
+{
+public:
+    virtual void StartupModule() override;
+    virtual void ShutdownModule() override;
+};
